XRCreativeFramework: Non-unity fixes.

diff --git a/Engine/Plugins/Experimental/XRCreativeFramework/Source/XRCreative/Private/ITF/SelectionInteraction.cpp b/Engine/Plugins/Experimental/XRCreativeFramework/Source/XRCreative/Private/ITF/SelectionInteraction.cpp
--- a/Engine/Plugins/Experimental/XRCreativeFramework/Source/XRCreative/Private/ITF/SelectionInteraction.cpp
+++ b/Engine/Plugins/Experimental/XRCreativeFramework/Source/XRCreative/Private/ITF/SelectionInteraction.cpp
@@ -1,8 +1,16 @@
 // Copyright Epic Games, Inc. All Rights Reserved.
 
 #include "SelectionInteraction.h"
+#include "CollisionQueryParams.h"
 #include "Elements/Framework/EngineElementsLibrary.h"
 #include "Elements/Framework/TypedElementSelectionSet.h"
+#include "Engine/HitResult.h"
+#include "Engine/World.h"
+#include "GameFramework/Actor.h"
+
+#if WITH_EDITOR
+#	include "Editor.h"
+#endif
 
 
 void UXRCreativeSelectionInteraction::Initialize(UTypedElementSelectionSet* InSelectionSet, TUniqueFunction<bool()> InCanChangeSelectionCallback)
diff --git a/Engine/Plugins/Experimental/XRCreativeFramework/Source/XRCreative/Private/ITF/XRCreativeITFComponent.cpp b/Engine/Plugins/Experimental/XRCreativeFramework/Source/XRCreative/Private/ITF/XRCreativeITFComponent.cpp
--- a/Engine/Plugins/Experimental/XRCreativeFramework/Source/XRCreative/Private/ITF/XRCreativeITFComponent.cpp
+++ b/Engine/Plugins/Experimental/XRCreativeFramework/Source/XRCreative/Private/ITF/XRCreativeITFComponent.cpp
@@ -29,6 +29,7 @@
 #	include "Editor/UnrealEdEngine.h"
 #	include "IVREditorModule.h"
 #	include "LevelEditorSubsystem.h"
+#	include "ScopedTransaction.h"
 #	include "SLevelViewport.h"
 #	include "UnrealEdGlobals.h"
 #	include "VREditorModeBase.h"
diff --git a/Engine/Plugins/Experimental/XRCreativeFramework/Source/XRCreative/Private/XRCreativeAvatar.cpp b/Engine/Plugins/Experimental/XRCreativeFramework/Source/XRCreative/Private/XRCreativeAvatar.cpp
--- a/Engine/Plugins/Experimental/XRCreativeFramework/Source/XRCreative/Private/XRCreativeAvatar.cpp
+++ b/Engine/Plugins/Experimental/XRCreativeFramework/Source/XRCreative/Private/XRCreativeAvatar.cpp
@@ -13,6 +13,7 @@
 #include "Engine/InputDelegateBinding.h"
 #include "EnhancedInputComponent.h"
 #include "EnhancedInputSubsystems.h"
+#include "Framework/Application/SlateApplication.h"
 #include "GameFramework/InputSettings.h"
 #include "GameFramework/PlayerController.h"
 #include "Haptics/HapticFeedbackEffect_Base.h"
diff --git a/Engine/Plugins/Experimental/XRCreativeFramework/Source/XRCreative/Private/XRCreativeSubsystem.cpp b/Engine/Plugins/Experimental/XRCreativeFramework/Source/XRCreative/Private/XRCreativeSubsystem.cpp
--- a/Engine/Plugins/Experimental/XRCreativeFramework/Source/XRCreative/Private/XRCreativeSubsystem.cpp
+++ b/Engine/Plugins/Experimental/XRCreativeFramework/Source/XRCreative/Private/XRCreativeSubsystem.cpp
@@ -3,6 +3,7 @@
 #include "XRCreativeSubsystem.h"
 #include "XRCreativeSettings.h"
 #include "Misc/CoreDelegates.h"
+#include "UObject/Package.h"
 
 #if WITH_EDITOR
 #	include "IVREditorModule.h"
diff --git a/Engine/Plugins/Experimental/XRCreativeFramework/Source/XRCreative/Private/XRCreativeTeleporter.cpp b/Engine/Plugins/Experimental/XRCreativeFramework/Source/XRCreative/Private/XRCreativeTeleporter.cpp
--- a/Engine/Plugins/Experimental/XRCreativeFramework/Source/XRCreative/Private/XRCreativeTeleporter.cpp
+++ b/Engine/Plugins/Experimental/XRCreativeFramework/Source/XRCreative/Private/XRCreativeTeleporter.cpp
@@ -1,8 +1,8 @@
 // Copyright Epic Games, Inc. All Rights Reserved.
 
 #include "XRCreativeTeleporter.h"
-
 #include "MotionControllerComponent.h"
+#include "Components/StaticMeshComponent.h"
 
 
 // Sets default values
diff --git a/Engine/Plugins/Experimental/XRCreativeFramework/Source/XRCreativeEditor/Private/XRCreativeVREditorMode.cpp b/Engine/Plugins/Experimental/XRCreativeFramework/Source/XRCreativeEditor/Private/XRCreativeVREditorMode.cpp
--- a/Engine/Plugins/Experimental/XRCreativeFramework/Source/XRCreativeEditor/Private/XRCreativeVREditorMode.cpp
+++ b/Engine/Plugins/Experimental/XRCreativeFramework/Source/XRCreativeEditor/Private/XRCreativeVREditorMode.cpp
@@ -2,6 +2,7 @@
 
 #include "XRCreativeVREditorMode.h"
 #include "Editor.h"
+#include "Framework/Application/SlateApplication.h"
 #include "Framework/Notifications/NotificationManager.h"
 #include "IVREditorModule.h"
 #include "IXRTrackingSystem.h"
