XRCreativeFramework: NoPCH UnrealGame build fix.

diff --git a/Engine/Plugins/Experimental/XRCreativeFramework/Source/XRCreative/Private/XRCreativeAvatar.cpp b/Engine/Plugins/Experimental/XRCreativeFramework/Source/XRCreative/Private/XRCreativeAvatar.cpp
--- a/Engine/Plugins/Experimental/XRCreativeFramework/Source/XRCreative/Private/XRCreativeAvatar.cpp
+++ b/Engine/Plugins/Experimental/XRCreativeFramework/Source/XRCreative/Private/XRCreativeAvatar.cpp
@@ -9,6 +9,7 @@
 #include "Components/SkeletalMeshComponent.h"
 #include "Components/WidgetComponent.h"
 #include "Components/WidgetInteractionComponent.h"
+#include "Engine/Level.h"
 #include "Engine/LocalPlayer.h"
 #include "Engine/InputDelegateBinding.h"
 #include "EnhancedInputComponent.h"
