Remove unused Settings reference from XRCreativeFramework

diff --git a/Engine/Plugins/Experimental/XRCreativeFramework/Source/XRCreative/Private/Module.cpp b/Engine/Plugins/Experimental/XRCreativeFramework/Source/XRCreative/Private/Module.cpp
--- a/Engine/Plugins/Experimental/XRCreativeFramework/Source/XRCreative/Private/Module.cpp
+++ b/Engine/Plugins/Experimental/XRCreativeFramework/Source/XRCreative/Private/Module.cpp
@@ -2,8 +2,6 @@
 
 #include "Module.h"
 #include "XRCreativeLog.h"
-#include "ISettingsModule.h"
-#include "ISettingsSection.h"
 
 
 #define LOCTEXT_NAMESPACE "FXRCreativeModule"
diff --git a/Engine/Plugins/Experimental/XRCreativeFramework/Source/XRCreative/XRCreative.Build.cs b/Engine/Plugins/Experimental/XRCreativeFramework/Source/XRCreative/XRCreative.Build.cs
--- a/Engine/Plugins/Experimental/XRCreativeFramework/Source/XRCreative/XRCreative.Build.cs
+++ b/Engine/Plugins/Experimental/XRCreativeFramework/Source/XRCreative/XRCreative.Build.cs
@@ -28,7 +28,6 @@ public class XRCreative : ModuleRules
 				"InteractiveToolsFramework",
 				"LevelSequence",
 				"RenderCore", // TODO: Kill this? Only for FlushRenderingCommands
-				"Settings",
 				"Slate",
 				"SlateCore",
 				"TypedElementFramework",
