XR Creative Framework + Virtual Scouting: Latest versions of plugins from sandbox.

#virtualized

diff --git a/Engine/Plugins/Experimental/VirtualScouting/Config/BaseVirtualScouting.ini b/Engine/Plugins/Experimental/VirtualScouting/Config/BaseVirtualScouting.ini
new file mode 100644
--- /dev/null
+++ b/Engine/Plugins/Experimental/VirtualScouting/Config/BaseVirtualScouting.ini
@@ -0,0 +1,2 @@
+[/Script/VirtualScouting.VirtualScoutingSettings]
+bUseImperial=True
diff --git a/Engine/Plugins/Experimental/VirtualScouting/Source/VirtualScouting/Private/Module.cpp b/Engine/Plugins/Experimental/VirtualScouting/Source/VirtualScouting/Private/Module.cpp
new file mode 100644
--- /dev/null
+++ b/Engine/Plugins/Experimental/VirtualScouting/Source/VirtualScouting/Private/Module.cpp
@@ -0,0 +1,60 @@
+﻿// Copyright Epic Games, Inc. All Rights Reserved.
+
+#include "Module.h"
+#include "VirtualScoutingLog.h"
+#include "VirtualScoutingSettings.h"
+#include "ISettingsModule.h"
+#include "ISettingsSection.h"
+
+
+#define LOCTEXT_NAMESPACE "FVirtualScoutingModule"
+
+
+DEFINE_LOG_CATEGORY(LogVirtualScouting);
+
+
+void FVirtualScoutingModule::StartupModule()
+{
+	//RegisterSettings();
+}
+
+
+void FVirtualScoutingModule::ShutdownModule()
+{
+	//UnregisterSettings();
+}
+
+
+void FVirtualScoutingModule::RegisterSettings() const
+{
+	ISettingsModule* SettingsModule = FModuleManager::GetModulePtr<ISettingsModule>("Settings");
+	if (SettingsModule != nullptr)
+	{
+		ISettingsSectionPtr SettingsSection = SettingsModule->RegisterSettings("Project", "Plugins", "VirtualScouting",
+			LOCTEXT("VirtualScoutingSettingsName", "Virtual Scouting"),
+			LOCTEXT("VirtualScoutingSettingsDescription", "Configure the Virtual Scouting settings."),
+			GetMutableDefault<UVirtualScoutingSettings>());
+	}
+}
+
+
+void FVirtualScoutingModule::UnregisterSettings() const
+{
+	ISettingsModule* SettingsModule = FModuleManager::GetModulePtr<ISettingsModule>("Settings");
+	if (SettingsModule != nullptr)
+	{
+		SettingsModule->UnregisterSettings("Project", "Plugins", "VirtualScouting");
+	}
+}
+
+
+bool FVirtualScoutingModule::OnSettingsModified()
+{
+	return true;
+}
+
+
+IMPLEMENT_MODULE(FVirtualScoutingModule, VirtualScouting)
+
+
+#undef LOCTEXT_NAMESPACE
diff --git a/Engine/Plugins/Experimental/VirtualScouting/Source/VirtualScouting/Private/Module.h b/Engine/Plugins/Experimental/VirtualScouting/Source/VirtualScouting/Private/Module.h
new file mode 100644
--- /dev/null
+++ b/Engine/Plugins/Experimental/VirtualScouting/Source/VirtualScouting/Private/Module.h
@@ -0,0 +1,24 @@
+﻿// Copyright Epic Games, Inc. All Rights Reserved.
+
+#pragma once
+
+#include "CoreMinimal.h"
+#include "Modules/ModuleManager.h"
+
+
+class FVirtualScoutingModule : public IModuleInterface
+{
+public:
+	virtual void StartupModule() override;
+	virtual void ShutdownModule() override;
+
+private:
+	/** Register the settings class with the settings module **/
+	void RegisterSettings() const;
+
+	/** Unregister the settings class with the settings module **/
+	void UnregisterSettings() const;
+	
+	/** TODO Handler for when Virtual Scouting settings are changed. */
+	bool OnSettingsModified();
+};
diff --git a/Engine/Plugins/Experimental/VirtualScouting/Source/VirtualScouting/Private/VirtualScoutingLog.h b/Engine/Plugins/Experimental/VirtualScouting/Source/VirtualScouting/Private/VirtualScoutingLog.h
new file mode 100644
--- /dev/null
+++ b/Engine/Plugins/Experimental/VirtualScouting/Source/VirtualScouting/Private/VirtualScoutingLog.h
@@ -0,0 +1,8 @@
+// Copyright Epic Games, Inc. All Rights Reserved.
+
+#pragma once
+
+#include "Logging/LogMacros.h"
+
+
+DECLARE_LOG_CATEGORY_EXTERN(LogVirtualScouting, Log, All);
diff --git a/Engine/Plugins/Experimental/VirtualScouting/Source/VirtualScouting/Private/VirtualScoutingSettings.cpp b/Engine/Plugins/Experimental/VirtualScouting/Source/VirtualScouting/Private/VirtualScoutingSettings.cpp
new file mode 100644
--- /dev/null
+++ b/Engine/Plugins/Experimental/VirtualScouting/Source/VirtualScouting/Private/VirtualScoutingSettings.cpp
@@ -0,0 +1,9 @@
+// Copyright Epic Games, Inc. All Rights Reserved.
+
+#include "VirtualScoutingSettings.h"
+
+
+UVirtualScoutingSettings* UVirtualScoutingSettings::GetVirtualScoutingSettings()
+{
+	return GetMutableDefault<UVirtualScoutingSettings>();
+}
diff --git a/Engine/Plugins/Experimental/VirtualScouting/Source/VirtualScouting/Public/VirtualScoutingSettings.h b/Engine/Plugins/Experimental/VirtualScouting/Source/VirtualScouting/Public/VirtualScoutingSettings.h
new file mode 100644
--- /dev/null
+++ b/Engine/Plugins/Experimental/VirtualScouting/Source/VirtualScouting/Public/VirtualScoutingSettings.h
@@ -0,0 +1,26 @@
+// Copyright Epic Games, Inc. All Rights Reserved.
+
+#pragma once
+
+#include "Engine/DeveloperSettingsBackedByCVars.h"
+#include "VirtualScoutingSettings.generated.h"
+
+
+/**
+ * Virtual Scouting Settings 
+ */
+UCLASS(Config=VirtualScoutingSettings, DisplayName="Virtual Scouting")
+class VIRTUALSCOUTING_API UVirtualScoutingSettings : public UDeveloperSettingsBackedByCVars
+{
+	GENERATED_BODY()
+	
+public:
+	UPROPERTY(Config, EditAnywhere, BlueprintReadWrite, Category="Virtual Scouting", meta=(DisplayName="Show Measurements in Imperial Units"))
+	bool bUseImperial = false;
+
+	UPROPERTY(Config, EditAnywhere, BlueprintReadWrite, Category="Virtual Scouting", meta=(DisplayName="Camera Apertures"))
+	TArray<float> CameraApertureArray = {1.0,1.4,2.0,2.8,4.0,5.6,8.0,11.0,16.0,22.0};
+
+	UFUNCTION(BlueprintPure, Category="Virtual Scouting")
+	static UVirtualScoutingSettings* GetVirtualScoutingSettings();
+};
diff --git a/Engine/Plugins/Experimental/VirtualScouting/Source/VirtualScouting/VirtualScouting.Build.cs b/Engine/Plugins/Experimental/VirtualScouting/Source/VirtualScouting/VirtualScouting.Build.cs
new file mode 100644
--- /dev/null
+++ b/Engine/Plugins/Experimental/VirtualScouting/Source/VirtualScouting/VirtualScouting.Build.cs
@@ -0,0 +1,55 @@
+// Copyright Epic Games, Inc. All Rights Reserved.
+
+using UnrealBuildTool;
+
+public class VirtualScouting : ModuleRules
+{
+	public VirtualScouting(ReadOnlyTargetRules Target) : base(Target)
+	{
+		PCHUsage = ModuleRules.PCHUsageMode.UseExplicitOrSharedPCHs;
+		bTreatAsEngineModule = true;
+
+		PublicDependencyModuleNames.AddRange(
+			new string[]
+			{
+				"CommonUI",
+				"Core",
+				"XRCreative",
+			}
+		);
+
+		PrivateDependencyModuleNames.AddRange(
+			new string[]
+			{
+				"CoreUObject",
+				"DeveloperSettings",
+				"Engine",
+				"EnhancedInput",
+				"HeadMountedDisplay",
+				"InputCore",
+				"InteractiveToolsFramework",
+				"LevelSequence",
+				"RenderCore", // TODO: Kill this? Only for FlushRenderingCommands
+				"Settings",
+				"Slate",
+				"SlateCore",
+				"TypedElementFramework",
+				"TypedElementRuntime",
+				"UMG",
+			}
+		);
+
+		if (Target.bBuildEditor)
+		{
+			PrivateDependencyModuleNames.AddRange(new string[] {
+				"ConcertSyncClient",
+				"ConcertSyncCore",
+				"InputEditor",
+				"LevelEditor",
+				"MultiUserClient",
+				"UnrealEd",
+				"VREditor",
+			});
+		}
+	}
+}
diff --git a/Engine/Plugins/Experimental/VirtualScouting/Source/VirtualScoutingEditor/Private/VPScoutingMode.cpp b/Engine/Plugins/Experimental/VirtualScouting/Source/VirtualScoutingEditor/Private/VPScoutingMode.cpp
--- a/Engine/Plugins/Experimental/VirtualScouting/Source/VirtualScoutingEditor/Private/VPScoutingMode.cpp
+++ b/Engine/Plugins/Experimental/VirtualScouting/Source/VirtualScoutingEditor/Private/VPScoutingMode.cpp
@@ -1,6 +1,7 @@
 // Copyright Epic Games, Inc. All Rights Reserved.
 
 #include "VPScoutingMode.h"
+#include "VirtualScoutingEditorModule.h"
 #include "VirtualScoutingOpenXR.h"
 #include "VirtualScoutingOpenXRModule.h"
 
@@ -18,30 +19,10 @@
 #define LOCTEXT_NAMESPACE "VirtualScouting"
 
 
-DECLARE_LOG_CATEGORY_EXTERN(LogVirtualScouting, Log, All);
-DEFINE_LOG_CATEGORY(LogVirtualScouting);
-
-
 static const FName OpenXRSystemName = FName(TEXT("OpenXR"));
 
 
-class FVirtualScoutingEditorModule : public IModuleInterface
-{
-private:
-	virtual void StartupModule() override
-	{
-	}
-
-	virtual void ShutdownModule() override
-	{
-	}
-};
-
-
-IMPLEMENT_MODULE(FVirtualScoutingEditorModule, VirtualScoutingEditor);
-
-
-UVPScoutingMode::UVPScoutingMode(const FObjectInitializer& ObjectInitializer)
+UVPScoutingMode::UVPScoutingMode()
 {
 	InteractorClass = FSoftObjectPath("/VirtualProductionUtilities/VR/VirtualScoutingInteractor.VirtualScoutingInteractor_C");
 	TeleporterClass = FSoftObjectPath("/VirtualProductionUtilities/VR/VirtualScoutingTeleporter.VirtualScoutingTeleporter_C");
@@ -70,7 +51,7 @@ void UVPScoutingMode::Enter()
 	TSharedPtr<FVirtualScoutingOpenXRExtension> XrExt = FVirtualScoutingOpenXRModule::Get().GetOpenXRExt();
 	if (!XrExt)
 	{
-		UE_LOG(LogVirtualScouting, Error, TEXT("OpenXR extension plugin invalid"));
+		UE_LOG(LogVirtualScoutingEditor, Error, TEXT("OpenXR extension plugin invalid"));
 		IVREditorModule::Get().EnableVREditor(false);
 		return;
 	}
@@ -99,13 +80,13 @@ void UVPScoutingMode::Enter()
 			UVPScoutingMode* This = WeakThis.Get();
 			if (!This || !This->bIsFullyInitialized)
 			{
-				UE_LOG(LogVirtualScouting, Warning, TEXT("Stale UVPScoutingMode; ignoring HmdDeviceType future"));
+				UE_LOG(LogVirtualScoutingEditor, Warning, TEXT("Stale UVPScoutingMode; ignoring HmdDeviceType future"));
 				return DeviceType;
 			}
 
 			if (DeviceType == NAME_None)
 			{
-				UE_LOG(LogVirtualScouting, Error, TEXT("Unable to map legacy HMD device type"));
+				UE_LOG(LogVirtualScoutingEditor, Error, TEXT("Unable to map legacy HMD device type"));
 				return DeviceType;
 			}
 
@@ -140,7 +121,7 @@ void UVPScoutingMode::InvalidSettingNotification(const FText& ErrorDetails)
 {
 	FText NotificationHeading = LOCTEXT("InvalidSettingNotificationTitle", "Unable to initialize Virtual Scouting");
 
-	UE_LOG(LogVirtualScouting, Error, TEXT("%s: %s"), *NotificationHeading.ToString(), *ErrorDetails.ToString());
+	UE_LOG(LogVirtualScoutingEditor, Error, TEXT("%s: %s"), *NotificationHeading.ToString(), *ErrorDetails.ToString());
 
 	FNotificationInfo Info(NotificationHeading);
 	Info.SubText = ErrorDetails;
diff --git a/Engine/Plugins/Experimental/VirtualScouting/Source/VirtualScoutingEditor/Private/VPScoutingMode.h b/Engine/Plugins/Experimental/VirtualScouting/Source/VirtualScoutingEditor/Private/VPScoutingMode.h
--- a/Engine/Plugins/Experimental/VirtualScouting/Source/VirtualScoutingEditor/Private/VPScoutingMode.h
+++ b/Engine/Plugins/Experimental/VirtualScouting/Source/VirtualScoutingEditor/Private/VPScoutingMode.h
@@ -6,19 +6,17 @@
 #include "VPScoutingMode.generated.h"
 
 
-UCLASS(meta=(DisplayName="Virtual Scouting Mode"))
+UCLASS(meta=(DisplayName="Virtual Scouting (Legacy VPUtilities)"))
 class UVPScoutingMode : public UVREditorMode
 {
 	GENERATED_BODY()
 
 public:
-	UVPScoutingMode(const FObjectInitializer& ObjectInitializer);
+	UVPScoutingMode();
 
 	//~ Begin UVREditorMode interface
 	virtual bool NeedsSyntheticDpad() override;
 
-	virtual bool ShouldDisplayExperimentalWarningOnEntry() const override { return false; }
-
 	virtual void Enter() override;
 	//~ End UVREditorMode interface
 
diff --git a/Engine/Plugins/Experimental/VirtualScouting/Source/VirtualScoutingEditor/Private/VirtualScoutingBlueprintFunctionLibrary.cpp b/Engine/Plugins/Experimental/VirtualScouting/Source/VirtualScoutingEditor/Private/VirtualScoutingBlueprintFunctionLibrary.cpp
new file mode 100644
--- /dev/null
+++ b/Engine/Plugins/Experimental/VirtualScouting/Source/VirtualScoutingEditor/Private/VirtualScoutingBlueprintFunctionLibrary.cpp
@@ -0,0 +1,54 @@
+// Copyright Epic Games, Inc. All Rights Reserved.
+
+
+#include "VirtualScoutingBlueprintFunctionLibrary.h"
+
+#include "Editor.h"
+#include "ScopedTransaction.h"
+#include "Editor/UnrealEdEngine.h"
+#include "Subsystems/EditorActorSubsystem.h"
+#include "Subsystems/UnrealEditorSubsystem.h"
+
+bool UVirtualScoutingBlueprintFunctionLibrary::CheckIsWithEditor()
+{
+
+#if WITH_EDITOR
+	return true;
+#endif
+	
+	return false;
+}
+
+bool UVirtualScoutingBlueprintFunctionLibrary::DeleteActors(const TArray<AActor*>& InActorsToDelete)
+{
+	const FScopedTransaction Transaction(NSLOCTEXT("VirtualScouting", "DeleteActors", "Delete Actors"));
+
+	UWorld* World = nullptr;
+	UUnrealEditorSubsystem* UnrealEditorSubsystem = GEditor->GetEditorSubsystem<UUnrealEditorSubsystem>();
+	World = UnrealEditorSubsystem->GetEditorWorld();
+	if (World)
+	{
+		// select
+		UEditorActorSubsystem* EditorActorSubsystem = GEditor->GetEditorSubsystem<UEditorActorSubsystem>();
+		if (EditorActorSubsystem)
+		{
+			EditorActorSubsystem->SetSelectedLevelActors(InActorsToDelete);
+		}
+		// delete selected
+		const bool bVerifyDeletionCanHappen = true;
+		const bool bWarnAboutReferences = false;
+		return GEditor->edactDeleteSelected(World, bVerifyDeletionCanHappen, bWarnAboutReferences, bWarnAboutReferences);
+	}
+
+	// If EditorWorld is not valid (i.e. VRPreviewMode), just call DestroyActor
+	bool bDeleted = false;
+	for (AActor* Actor : InActorsToDelete)
+	{
+		World = Actor->GetWorld();
+		if (World)
+		{
+			bDeleted |= World->DestroyActor(Actor);
+		}
+	}
+	return bDeleted;
+}
\ No newline at end of file
diff --git a/Engine/Plugins/Experimental/VirtualScouting/Source/VirtualScoutingEditor/Private/VirtualScoutingEditorModule.cpp b/Engine/Plugins/Experimental/VirtualScouting/Source/VirtualScoutingEditor/Private/VirtualScoutingEditorModule.cpp
new file mode 100644
--- /dev/null
+++ b/Engine/Plugins/Experimental/VirtualScouting/Source/VirtualScoutingEditor/Private/VirtualScoutingEditorModule.cpp
@@ -0,0 +1,23 @@
+// Copyright Epic Games, Inc. All Rights Reserved.
+
+#include "VirtualScoutingEditorModule.h"
+#include "Modules/ModuleManager.h"
+
+
+DEFINE_LOG_CATEGORY(LogVirtualScoutingEditor);
+
+
+class FVirtualScoutingEditorModule : public IModuleInterface
+{
+private:
+	virtual void StartupModule() override
+	{
+	}
+
+	virtual void ShutdownModule() override
+	{
+	}
+};
+
+
+IMPLEMENT_MODULE(FVirtualScoutingEditorModule, VirtualScoutingEditor);
diff --git a/Engine/Plugins/Experimental/VirtualScouting/Source/VirtualScoutingEditor/Private/VirtualScoutingEditorModule.h b/Engine/Plugins/Experimental/VirtualScouting/Source/VirtualScoutingEditor/Private/VirtualScoutingEditorModule.h
new file mode 100644
--- /dev/null
+++ b/Engine/Plugins/Experimental/VirtualScouting/Source/VirtualScoutingEditor/Private/VirtualScoutingEditorModule.h
@@ -0,0 +1,8 @@
+// Copyright Epic Games, Inc. All Rights Reserved.
+
+#pragma once
+
+#include "Logging/LogMacros.h"
+
+
+DECLARE_LOG_CATEGORY_EXTERN(LogVirtualScoutingEditor, Log, All);
diff --git a/Engine/Plugins/Experimental/VirtualScouting/Source/VirtualScoutingEditor/Public/VirtualScoutingBlueprintFunctionLibrary.h b/Engine/Plugins/Experimental/VirtualScouting/Source/VirtualScoutingEditor/Public/VirtualScoutingBlueprintFunctionLibrary.h
new file mode 100644
--- /dev/null
+++ b/Engine/Plugins/Experimental/VirtualScouting/Source/VirtualScoutingEditor/Public/VirtualScoutingBlueprintFunctionLibrary.h
@@ -0,0 +1,24 @@
+// Copyright Epic Games, Inc. All Rights Reserved.
+
+#pragma once
+
+#include "CoreMinimal.h"
+#include "Kismet/BlueprintFunctionLibrary.h"
+#include "VirtualScoutingBlueprintFunctionLibrary.generated.h"
+
+/**
+ * 
+ */
+UCLASS()
+class VIRTUALSCOUTINGEDITOR_API UVirtualScoutingBlueprintFunctionLibrary : public UBlueprintFunctionLibrary
+{
+	GENERATED_BODY()
+
+	UFUNCTION(BlueprintCallable, Category="VirtualScouting")
+	static bool CheckIsWithEditor();
+
+	/** Helper function to delete actors. This skips reference checks.*/
+	UFUNCTION(BlueprintCallable, Category="VirtualScouting")
+	static bool DeleteActors(const TArray<AActor*>& InActorsToDelete);
+
+};
diff --git a/Engine/Plugins/Experimental/VirtualScouting/Source/VirtualScoutingEditor/VirtualScoutingEditor.Build.cs b/Engine/Plugins/Experimental/VirtualScouting/Source/VirtualScoutingEditor/VirtualScoutingEditor.Build.cs
--- a/Engine/Plugins/Experimental/VirtualScouting/Source/VirtualScoutingEditor/VirtualScoutingEditor.Build.cs
+++ b/Engine/Plugins/Experimental/VirtualScouting/Source/VirtualScoutingEditor/VirtualScoutingEditor.Build.cs
@@ -7,6 +7,7 @@ public class VirtualScoutingEditor : ModuleRules
 	public VirtualScoutingEditor(ReadOnlyTargetRules Target) : base(Target)
 	{
 		PCHUsage = ModuleRules.PCHUsageMode.UseExplicitOrSharedPCHs;
+		bTreatAsEngineModule = true;
 
 		PrivateIncludePaths.AddRange(
 			new string[]
@@ -26,10 +27,22 @@ public class VirtualScoutingEditor : ModuleRules
 		PrivateDependencyModuleNames.AddRange(new string[]
 			{
 				"OpenXRHMD",
-				"Slate",
+				"OpenXRInput",
 				"UnrealEd",
 				"VREditor",
 				"VirtualScoutingOpenXR",
+				"EnhancedInput",
+				"HeadMountedDisplay",
+				"InputCore",
+				"InputEditor",
+				"InteractiveToolsFramework",
+				"LevelEditor",
+				"RenderCore",
+				"Slate",
+				"SlateCore",
+				"UMG",
+				"VirtualScouting",
+				"XRCreative",
 			}
 		);
 	}
diff --git a/Engine/Plugins/Experimental/VirtualScouting/Source/VirtualScoutingOpenXR/VirtualScoutingOpenXR.Build.cs b/Engine/Plugins/Experimental/VirtualScouting/Source/VirtualScoutingOpenXR/VirtualScoutingOpenXR.Build.cs
--- a/Engine/Plugins/Experimental/VirtualScouting/Source/VirtualScoutingOpenXR/VirtualScoutingOpenXR.Build.cs
+++ b/Engine/Plugins/Experimental/VirtualScouting/Source/VirtualScoutingOpenXR/VirtualScoutingOpenXR.Build.cs
@@ -7,6 +7,7 @@ public class VirtualScoutingOpenXR : ModuleRules
 	public VirtualScoutingOpenXR(ReadOnlyTargetRules Target) : base(Target)
 	{
 		PCHUsage = ModuleRules.PCHUsageMode.UseExplicitOrSharedPCHs;
+		bTreatAsEngineModule = true;
 
 		PublicDependencyModuleNames.AddRange(new string[]
 			{
diff --git a/Engine/Plugins/Experimental/VirtualScouting/VirtualScouting.uplugin b/Engine/Plugins/Experimental/VirtualScouting/VirtualScouting.uplugin
--- a/Engine/Plugins/Experimental/VirtualScouting/VirtualScouting.uplugin
+++ b/Engine/Plugins/Experimental/VirtualScouting/VirtualScouting.uplugin
@@ -11,11 +11,12 @@
 	"MarketplaceURL": "",
 	"SupportURL": "",
 	"CanContainContent": true,
+	"IsBetaVersion": false,
 	"IsExperimentalVersion": true,
 	"Installed": false,
+	"EnabledByDefault": false,
 	"SupportedTargetPlatforms": [
-		"Win64",
-		"Linux"
+		"Win64"
 	],
 	"Modules": [
 		{
@@ -23,8 +24,15 @@
 			"Type": "Runtime",
 			"LoadingPhase": "PostConfigInit",
 			"PlatformAllowList": [
-				"Win64",
-				"Linux"
+				"Win64"
+			]
+		},
+		{
+			"Name": "VirtualScouting",
+			"Type": "Runtime",
+			"LoadingPhase": "Default",
+			"PlatformAllowList": [
+				"Win64"
 			]
 		},
 		{
@@ -32,19 +40,66 @@
 			"Type": "Editor",
 			"LoadingPhase": "Default",
 			"PlatformAllowList": [
-				"Win64",
-				"Linux"
+				"Win64"
 			]
 		}
 	],
 	"Plugins": [
 		{
-			"Name": "VirtualProductionUtilities",
-			"Enabled": true
+			"Name": "XRCreativeFramework",
+			"Enabled":  true
 		},
 		{
 			"Name": "OpenXR",
 			"Enabled": true
+		},
+		{
+			"Name": "EnhancedInput",
+			"Enabled": true
+		},
+		{
+			"Name": "CommonUI",
+			"Enabled": true
+		},
+		{
+			"Name": "BlueprintFileUtils",
+			"Enabled": true
+		},
+		{
+			"Name": "JsonBlueprintUtilities",
+			"Enabled": true
+		},
+		{
+			"Name": "SunPosition",
+			"Enabled": true
+		},
+		{
+			"Name": "GeometryScripting",
+			"Enabled": true
+		},
+		{
+			"Name": "GizmoFramework",
+			"Enabled": true
+		},
+		{
+			"Name": "GizmoEdMode",
+			"Enabled": true
+		},
+		{
+			"Name": "ModelingToolsEditorMode",
+			"Enabled": true
+		},
+		{
+			"Name": "ConsoleVariables",
+			"Enabled": true
+		},
+		{
+			"Name": "LevelSnapshots",
+			"Enabled": true
+		},
+		{
+			"Name": "MultiUserClient",
+			"Enabled": true
 		}
 	]
 }
diff --git a/Engine/Plugins/Experimental/XRCreativeFramework/Source/XRCreative/Private/ITF/SelectionInteraction.cpp b/Engine/Plugins/Experimental/XRCreativeFramework/Source/XRCreative/Private/ITF/SelectionInteraction.cpp
--- a/Engine/Plugins/Experimental/XRCreativeFramework/Source/XRCreative/Private/ITF/SelectionInteraction.cpp
+++ b/Engine/Plugins/Experimental/XRCreativeFramework/Source/XRCreative/Private/ITF/SelectionInteraction.cpp
@@ -13,12 +13,12 @@
 #endif
 
 
-void UXRCreativeSelectionInteraction::Initialize(UTypedElementSelectionSet* InSelectionSet, TUniqueFunction<bool()> InCanChangeSelectionCallback)
+void UXRCreativeSelectionInteraction::Initialize(UTypedElementSelectionSet* InSelectionSet, FActorPredicate InCanSelectCallback)
 {
 	ensure(InSelectionSet);
 
 	WeakSelectionSet = InSelectionSet;
-	CanChangeSelectionCallback = MoveTemp(InCanChangeSelectionCallback);
+	CanSelectCallback = MoveTemp(InCanSelectCallback);
 
 	// create click behavior and set ourselves as click target
 	ClickBehavior = NewObject<USingleClickInputBehavior>();
@@ -61,21 +61,17 @@ FInputRayHit UXRCreativeSelectionInteraction::IsHitByClick(const FInputDeviceRay
 {
 	FInputRayHit RayHit;
 
-	if (CanChangeSelectionCallback() == false)
-	{
-		return RayHit;
-	}
-
-	FCollisionObjectQueryParams QueryParams(FCollisionObjectQueryParams::AllObjects);
 	FHitResult Result;
+	FCollisionObjectQueryParams QueryParams(FCollisionObjectQueryParams::AllObjects);
 	const bool bBlockingHit = GetWorld()->LineTraceSingleByObjectType(Result,
 		ClickPos.WorldRay.Origin, ClickPos.WorldRay.PointAt(999999), QueryParams);
-	if (bBlockingHit)
+	AActor* HitActor = Result.GetActor();
+	if (bBlockingHit && CanSelectCallback(HitActor))
 	{
 		RayHit.bHit = true;
 		RayHit.HitDepth = Result.Distance;
 		RayHit.bHasHitNormal = true;
-		RayHit.SetHitObject(Result.GetActor());
+		RayHit.SetHitObject(HitActor);
 	}
 
 	return RayHit;
@@ -116,8 +112,10 @@ void UXRCreativeSelectionInteraction::OnClicked(const FInputDeviceRay& ClickPos)
 	FHitResult Result;
 	const bool bBlockingHit = GetWorld()->LineTraceSingleByObjectType(Result,
 		ClickPos.WorldRay.Origin, ClickPos.WorldRay.PointAt(999999), QueryParams);
-	if (bBlockingHit)
+	if (ensure(bBlockingHit))
 	{
+		ensure(CanSelectCallback(Result.GetActor()));
+
 		FTypedElementHandle HitActorHandle;
 
 #if WITH_EDITOR
diff --git a/Engine/Plugins/Experimental/XRCreativeFramework/Source/XRCreative/Private/ITF/SelectionInteraction.h b/Engine/Plugins/Experimental/XRCreativeFramework/Source/XRCreative/Private/ITF/SelectionInteraction.h
--- a/Engine/Plugins/Experimental/XRCreativeFramework/Source/XRCreative/Private/ITF/SelectionInteraction.h
+++ b/Engine/Plugins/Experimental/XRCreativeFramework/Source/XRCreative/Private/ITF/SelectionInteraction.h
@@ -19,16 +19,17 @@ class UXRCreativeSelectionInteraction : public UObject, public IInputBehaviorSou
 	GENERATED_BODY()
 
 public:
+	using FActorPredicate = TUniqueFunction<bool(AActor*)>;
+
 	/**
 	 * Set up the Interaction, creates and registers Behaviors/etc. 
 	 * 
 	 * @param InSelectionSet the typed element selection set we maintain a weak pointer to and operate on
-	 * @param InCanChangeSelectionCallback this function will be called to determine if the current Selection is allowed to be modified (for example, when a Tool is active, we may wish to lock selection)
+	 * @param InCanSelectCallback this function will be called to determine if the selection can be changed to the specified actor (or null); when a tool is active, we may wish to lock selection
 	 */
-	void Initialize(UTypedElementSelectionSet* InSelectionSet, TUniqueFunction<bool()> InCanChangeSelectionCallback);
+	void Initialize(UTypedElementSelectionSet* InSelectionSet, FActorPredicate InCanSelectCallback);
 	void Shutdown();
 
-public:
 	// IInputBehaviorSource interface
 	virtual const UInputBehaviorSet* GetInputBehaviors() const { return BehaviorSet; }
 
@@ -57,8 +58,8 @@ protected:
 
 	TWeakObjectPtr<UTypedElementSelectionSet> WeakSelectionSet;
 
-	// default change-selection callback always allows selection change
-	TUniqueFunction<bool()> CanChangeSelectionCallback = []() { return true; };
+	// default predicate allows anything
+	FActorPredicate CanSelectCallback = [](AActor*) { return true; };
 
 	// flags used to identify behavior modifier keys/buttons
 	static const int AddToSelectionModifier = 1;
diff --git a/Engine/Plugins/Experimental/XRCreativeFramework/Source/XRCreative/Private/ITF/XRCreativeITFComponent.cpp b/Engine/Plugins/Experimental/XRCreativeFramework/Source/XRCreative/Private/ITF/XRCreativeITFComponent.cpp
--- a/Engine/Plugins/Experimental/XRCreativeFramework/Source/XRCreative/Private/ITF/XRCreativeITFComponent.cpp
+++ b/Engine/Plugins/Experimental/XRCreativeFramework/Source/XRCreative/Private/ITF/XRCreativeITFComponent.cpp
@@ -1,6 +1,7 @@
 // Copyright Epic Games, Inc. All Rights Reserved.
 
 #include "XRCreativeITFComponent.h"
+#include "XRCreativeAvatar.h"
 #include "XRCreativeITFRenderComponent.h"
 #include "XRCreativeLog.h"
 #include "XRCreativePointerComponent.h"
@@ -93,6 +94,11 @@ public:
 		return ToolsComp->GetCurrentCoordinateSystem();
 	}
 
+	virtual EToolContextTransformGizmoMode GetCurrentTransformGizmoMode() const override
+	{
+		return ToolsComp->GetCurrentTransformGizmoMode();
+	}
+
 	virtual FToolContextSnappingConfiguration GetCurrentSnappingSettings() const override
 	{
 		return FToolContextSnappingConfiguration();
@@ -356,6 +362,8 @@ UXRCreativeITFComponent::UXRCreativeITFComponent()
 
 	bAutoActivate = true;
 	bWantsInitializeComponent = true;
+
+	UnselectableActorClasses.Add(AXRCreativeAvatar::StaticClass());
 }
 
 
@@ -397,7 +405,31 @@ void UXRCreativeITFComponent::InitializeComponent()
 	// register selection interaction
 	SelectionInteraction = NewObject<UXRCreativeSelectionInteraction>(this);
 	SelectionInteraction->Initialize(GetSelectionSet(),
-		[this]() { return HaveActiveTool() == false; }
+		[this](AActor* SelectionCandidate)
+		{
+			if (HaveActiveTool())
+			{
+				return false;
+			}
+
+			if (CanSelectPredicate.IsBound())
+			{
+				return CanSelectPredicate.Execute(SelectionCandidate);
+			}
+
+			if (SelectionCandidate)
+			{
+				for (const TSubclassOf<AActor>& DisallowedClass : UnselectableActorClasses)
+				{
+					if (DisallowedClass.Get() && SelectionCandidate->IsA(DisallowedClass))
+					{
+						return false;
+					}
+				}
+			}
+
+			return true;
+		}
 	);
 	ToolsContext->InputRouter->RegisterSource(SelectionInteraction);
 
@@ -528,6 +560,7 @@ void UXRCreativeITFComponent::TickComponent(float InDeltaTime, enum ELevelTick I
 #if WITH_EDITOR
 		if (IsInEditor())
 		{
+			FEditorScriptExecutionGuard ScriptGuard;
 			EditorToolsTick(InDeltaTime);
 		}
 #endif
@@ -584,8 +617,9 @@ void UXRCreativeITFComponent::ToolsTick(float InDeltaTime)
 	ULocalPlayer* LocalPlayer = Cast<ULocalPlayer>(PC->Player);
 	FVector ViewLocation;
 	FRotator ViewRotation;
+	const int32 StereoViewIndex = GEngine->IsStereoscopic3D(Viewport) ? EStereoscopicEye::eSSE_LEFT_EYE : INDEX_NONE;
 	FSceneView* SceneView = LocalPlayer->CalcSceneView(&ViewFamily, ViewLocation, ViewRotation,
-		LocalPlayer->ViewportClient->Viewport, nullptr, EStereoscopicEye::eSSE_LEFT_EYE);
+		LocalPlayer->ViewportClient->Viewport, nullptr, StereoViewIndex);
 	if (!SceneView)
 	{
 		return;
@@ -638,8 +672,8 @@ void UXRCreativeITFComponent::ToolsTick(float InDeltaTime)
 	//InputState.Mouse.WorldRay = FRay(Origin, Direction);
 #endif // #ifdef XRCREATIVE_CALC_ITF_MOUSE_2D
 
-	InputState.Mouse.WorldRay = FRay(PointerComponent->GetComponentLocation(),
-	                                 PointerComponent->GetFilteredTraceEnd(false));
+	const FVector PointerDirection = (PointerComponent->GetFilteredTraceEnd(false) - PointerComponent->GetComponentLocation()).GetSafeNormal();
+	InputState.Mouse.WorldRay = FRay(PointerComponent->GetComponentLocation(), PointerDirection, true);
 
 	if (bPendingMouseStateChange || ToolsContext->InputRouter->HasActiveMouseCapture())
 	{
@@ -749,8 +783,8 @@ void UXRCreativeITFComponent::EditorToolsTick(float InDeltaTime)
 	//InputState.Mouse.WorldRay = FRay(RayOrigin, RayDirection);
 #endif // #ifdef XRCREATIVE_CALC_ITF_MOUSE_2D
 
-	InputState.Mouse.WorldRay = FRay(PointerComponent->GetComponentLocation(),
-	                                 PointerComponent->GetFilteredTraceEnd(false));
+	const FVector PointerDirection = (PointerComponent->GetFilteredTraceEnd(false) - PointerComponent->GetComponentLocation()).GetSafeNormal();
+	InputState.Mouse.WorldRay = FRay(PointerComponent->GetComponentLocation(), PointerDirection, true);
 
 	if (bPendingMouseStateChange || ToolsContext->InputRouter->HasActiveMouseCapture())
 	{
@@ -792,6 +826,11 @@ void UXRCreativeITFComponent::EditorToolsTick(float InDeltaTime)
 
 void UXRCreativeITFComponent::LeftMousePressed()
 {
+	if (!IsActive())
+	{
+		return;
+	}
+
 	CurrentMouseState.Mouse.Left.SetStates(true, false, false);
 	bPendingMouseStateChange = true;
 }
@@ -799,6 +838,11 @@ void UXRCreativeITFComponent::LeftMousePressed()
 
 void UXRCreativeITFComponent::LeftMouseReleased()
 {
+	if (!IsActive())
+	{
+		return;
+	}
+
 	CurrentMouseState.Mouse.Left.SetStates(false, false, true);
 	bPendingMouseStateChange = true;
 }
@@ -836,7 +880,18 @@ bool UXRCreativeITFComponent::HaveActiveTool()
 
 
 void UXRCreativeITFComponent::SetCurrentCoordinateSystem(EToolContextCoordinateSystem CoordSystem)
-{ 
-	CurrentCoordinateSystem = CoordSystem; 
-	//TransformInteraction->ForceUpdateGizmoState();
+{
+	if(IsValid(GetSelectionSet()))
+	{
+		CurrentCoordinateSystem = CoordSystem; 
+		TransformInteraction->ForceUpdateGizmoState();	
+	}
+
+}
+
+
+void UXRCreativeITFComponent::SetCurrentTransformGizmoMode(EToolContextTransformGizmoMode GizmoMode)
+{
+	CurrentTransformGizmoMode = GizmoMode;
+	TransformInteraction->ForceUpdateGizmoState();
 }
diff --git a/Engine/Plugins/Experimental/XRCreativeFramework/Source/XRCreative/Private/XRCreativeAvatar.cpp b/Engine/Plugins/Experimental/XRCreativeFramework/Source/XRCreative/Private/XRCreativeAvatar.cpp
--- a/Engine/Plugins/Experimental/XRCreativeFramework/Source/XRCreative/Private/XRCreativeAvatar.cpp
+++ b/Engine/Plugins/Experimental/XRCreativeFramework/Source/XRCreative/Private/XRCreativeAvatar.cpp
@@ -192,13 +192,6 @@ void AXRCreativeAvatar::ConfigureToolset(UXRCreativeToolset* InToolset)
 			Tools.Add(NewTool);
 		}
 	}
-
-	if (ensure(Toolset->Palette))
-	{
-		UXRCreativePalette* Palette = CreateWidget<UXRCreativePalette>(GetWorld(), Toolset->Palette);
-		Palette->SetOwner(this);
-		MenuWidget->SetWidget(Palette);
-	}
 }
 
 
@@ -302,7 +295,6 @@ void AXRCreativeAvatar::RemoveInputMappingContext(UInputMappingContext* InContex
 	}
 }
 
-
 //** Cues up a Haptic effect to be processed on tick in ProcessHaptics() //
 void AXRCreativeAvatar::PlayHapticEffect(UHapticFeedbackEffect_Base* HapticEffect, const int ControllerID, const EControllerHand Hand, float Scale, bool bLoop)
 {
diff --git a/Engine/Plugins/Experimental/XRCreativeFramework/Source/XRCreative/Private/XRCreativePointerComponent.cpp b/Engine/Plugins/Experimental/XRCreativeFramework/Source/XRCreative/Private/XRCreativePointerComponent.cpp
--- a/Engine/Plugins/Experimental/XRCreativeFramework/Source/XRCreative/Private/XRCreativePointerComponent.cpp
+++ b/Engine/Plugins/Experimental/XRCreativeFramework/Source/XRCreative/Private/XRCreativePointerComponent.cpp
@@ -55,7 +55,10 @@ void UXRCreativePointerComponent::TickComponent(float InDeltaTime, enum ELevelTi
 	const FVector FilteredTraceEnd = SmoothingFilter->Filter(RawTraceEnd, InDeltaTime);
 
 	FCollisionQueryParams QueryParams;
-	QueryParams.AddIgnoredActor(GetOwner());
+	QueryParams.AddIgnoredActors(IgnoredActors);
+	QueryParams.AddIgnoredComponents(IgnoredComponents);
+	QueryParams.bTraceComplex = true;
+
 	GetWorld()->LineTraceSingleByChannel(HitResult, GetComponentLocation(), FilteredTraceEnd, ECC_Visibility, QueryParams);
 }
 
@@ -65,7 +68,7 @@ FVector UXRCreativePointerComponent::GetRawTraceEnd(const bool bScaledByImpact /
 	if (bScaledByImpact)
 	{
 		const FVector TraceStart = GetComponentLocation();
-		return ((RawTraceEnd - TraceStart) * HitResult.Time) + TraceStart;
+		return TraceStart + ((RawTraceEnd - TraceStart) * HitResult.Time);
 	}
 	else
 	{
@@ -80,7 +83,7 @@ FVector UXRCreativePointerComponent::GetFilteredTraceEnd(const bool bScaledByImp
 	if (bScaledByImpact)
 	{
 		const FVector TraceStart = GetComponentLocation();
-		return ((HitResult.TraceEnd - TraceStart) * HitResult.Time) + TraceStart;
+		return TraceStart + ((HitResult.TraceEnd - TraceStart) * HitResult.Time);
 	}
 	else
 	{
diff --git a/Engine/Plugins/Experimental/XRCreativeFramework/Source/XRCreative/Private/XRCreativeSubsystem.cpp b/Engine/Plugins/Experimental/XRCreativeFramework/Source/XRCreative/Private/XRCreativeSubsystem.cpp
--- a/Engine/Plugins/Experimental/XRCreativeFramework/Source/XRCreative/Private/XRCreativeSubsystem.cpp
+++ b/Engine/Plugins/Experimental/XRCreativeFramework/Source/XRCreative/Private/XRCreativeSubsystem.cpp
@@ -3,6 +3,7 @@
 #include "XRCreativeSubsystem.h"
 #include "XRCreativeSettings.h"
 #include "Misc/CoreDelegates.h"
+#include "Types/MVVMViewModelCollection.h"
 #include "UObject/Package.h"
 
 #if WITH_EDITOR
@@ -13,6 +14,7 @@
 void UXRCreativeSubsystem::Initialize(FSubsystemCollectionBase& Collection)
 {
 	Super::Initialize(Collection);
+	ViewModelCollection = NewObject<UMVVMViewModelCollectionObject>(this);
 
 	EngineInitCompleteDelegate = FCoreDelegates::OnFEngineLoopInitComplete.AddUObject(this, &UXRCreativeSubsystem::OnEngineInitComplete);
 }
@@ -48,8 +50,9 @@ void UXRCreativeSubsystem::OnEngineInitComplete()
 	EngineInitCompleteDelegate.Reset();
 
 	const UXRCreativeSettings* Settings = UXRCreativeSettings::GetXRCreativeSettings();
-	if (UClass* HelpersClass = Settings->SubsystemHelpersClass.LoadSynchronous())
+	if (UClass* HelperClass = Settings->SubsystemHelperClass.LoadSynchronous())
 	{
-		Helpers = NewObject<UXRCreativeSubsystemHelpers>(GetTransientPackage(), HelpersClass);
+		Helpers = NewObject<UXRCreativeSubsystemHelper>(GetTransientPackage(), HelperClass);
 	}
 }
+
diff --git a/Engine/Plugins/Experimental/XRCreativeFramework/Source/XRCreative/Private/XRCreativeTeleporter.cpp b/Engine/Plugins/Experimental/XRCreativeFramework/Source/XRCreative/Private/XRCreativeTeleporter.cpp
--- a/Engine/Plugins/Experimental/XRCreativeFramework/Source/XRCreative/Private/XRCreativeTeleporter.cpp
+++ b/Engine/Plugins/Experimental/XRCreativeFramework/Source/XRCreative/Private/XRCreativeTeleporter.cpp
@@ -2,6 +2,7 @@
 
 #include "XRCreativeTeleporter.h"
 #include "MotionControllerComponent.h"
+#include "XRDeviceVisualizationComponent.h"
 #include "Components/StaticMeshComponent.h"
 
 
@@ -27,6 +28,21 @@ AXRCreativeTeleporter::AXRCreativeTeleporter(const FObjectInitializer& ObjectIni
 	TeleportMesh = CreateDefaultSubobject<UStaticMeshComponent>("TeleportMesh");
 	TeleportMesh->SetupAttachment(GetRootComponent());
 	TeleportMesh->SetMobility(EComponentMobility::Movable);
+
+	HeadMountedDisplay = CreateDefaultSubobject<UMotionControllerComponent>("HeadMountedDisplay");
+	HeadMountedDisplay->SetupAttachment(GetRootComponent());
+	HeadMountedDisplay->bTickInEditor = true;
+	HeadMountedDisplay->MotionSource = "HMD";
+
+	LeftControllerVisual = CreateDefaultSubobject<UStaticMeshComponent>("LeftControllerVisual");
+	LeftControllerVisual->SetupAttachment(LeftController);
+
+	RightControllerVisual = CreateDefaultSubobject<UStaticMeshComponent>("RightControllerVisual");
+	RightControllerVisual->SetupAttachment(RightController);
+	
+	HMDVisual = CreateDefaultSubobject<UStaticMeshComponent>("HMDVisual");
+	HMDVisual->SetupAttachment(HeadMountedDisplay);
+	
 }
 
 // Called when the game starts or when spawned
diff --git a/Engine/Plugins/Experimental/XRCreativeFramework/Source/XRCreative/Private/XRCreativeToolActor.cpp b/Engine/Plugins/Experimental/XRCreativeFramework/Source/XRCreative/Private/XRCreativeToolActor.cpp
new file mode 100644
--- /dev/null
+++ b/Engine/Plugins/Experimental/XRCreativeFramework/Source/XRCreative/Private/XRCreativeToolActor.cpp
@@ -0,0 +1,26 @@
+// Copyright Epic Games, Inc. All Rights Reserved.
+
+
+#include "XRCreativeToolActor.h"
+
+
+// Sets default values
+AXRCreativeToolActor::AXRCreativeToolActor()
+{
+	// Set this actor to call Tick() every frame.  You can turn this off to improve performance if you don't need it.
+	PrimaryActorTick.bCanEverTick = true;
+}
+
+// Called when the game starts or when spawned
+void AXRCreativeToolActor::BeginPlay()
+{
+	Super::BeginPlay();
+	
+}
+
+// Called every frame
+void AXRCreativeToolActor::Tick(float DeltaTime)
+{
+	Super::Tick(DeltaTime);
+}
+
diff --git a/Engine/Plugins/Experimental/XRCreativeFramework/Source/XRCreative/Private/XRCreativeToolset.cpp b/Engine/Plugins/Experimental/XRCreativeFramework/Source/XRCreative/Private/XRCreativeToolset.cpp
--- a/Engine/Plugins/Experimental/XRCreativeFramework/Source/XRCreative/Private/XRCreativeToolset.cpp
+++ b/Engine/Plugins/Experimental/XRCreativeFramework/Source/XRCreative/Private/XRCreativeToolset.cpp
@@ -1,6 +1,23 @@
 // Copyright Epic Games, Inc. All Rights Reserved.
 
 #include "XRCreativeToolset.h"
+/*#include "XRCreativeSettings.h"
+#include "Logging/LogMacros.h"
 
+DEFINE_LOG_CATEGORY(LogXRCreativeToolset);
+
+UXRCreativeToolset::UXRCreativeToolset()
+{
+	UXRCreativeSettings* Settings = UXRCreativeSettings::GetXRCreativeSettings();
+	if (Settings->Handedness == EXRCreativeHandedness::Left && !LeftInputMappingContext)
+	{
+		UE_LOG(LogXRCreativeToolset, Warning, TEXT("Handedness is Left but no Left Input Mapping Context found in Toolset."));
+	}
+}
+
+UXRCreativeToolset::~UXRCreativeToolset()
+{
+	
+}*/
 
 
diff --git a/Engine/Plugins/Experimental/XRCreativeFramework/Source/XRCreative/Public/XRCreativeAvatar.h b/Engine/Plugins/Experimental/XRCreativeFramework/Source/XRCreative/Public/XRCreativeAvatar.h
--- a/Engine/Plugins/Experimental/XRCreativeFramework/Source/XRCreative/Public/XRCreativeAvatar.h
+++ b/Engine/Plugins/Experimental/XRCreativeFramework/Source/XRCreative/Public/XRCreativeAvatar.h
@@ -49,6 +49,8 @@ public:
 	virtual void BeginPlay() override;
 
 	void ConfigureToolset(UXRCreativeToolset* InToolset);
+	const UXRCreativeToolset* GetToolset() const { return Toolset; }
+	const TArray<TObjectPtr<UXRCreativeTool>>& GetTools() const { return Tools; }
 
 	UFUNCTION(BlueprintCallable, Category="XR Creative")
 	FTransform GetHeadTransform() const;
@@ -86,6 +88,12 @@ public:
 	UFUNCTION(BlueprintCallable, Category="XR Creative")
 	void RemoveInputMappingContext(UInputMappingContext* Context);
 
+	/**
+	* Called when In-Editor VR mode is started. In-Editor equivalent to Begin Play.
+	*/
+	
+	UFUNCTION(BlueprintImplementableEvent, BlueprintCallable, Category="XR Creative", meta=(DisplayName = "On Enter VR"))
+	void BP_OnVRInitialize();
 
 	/** Play haptic feedback asset on a given hand - only left and right supported
 	 * @param HapticEffect			The haptic effect to play
diff --git a/Engine/Plugins/Experimental/XRCreativeFramework/Source/XRCreative/Public/XRCreativeITFComponent.h b/Engine/Plugins/Experimental/XRCreativeFramework/Source/XRCreative/Public/XRCreativeITFComponent.h
--- a/Engine/Plugins/Experimental/XRCreativeFramework/Source/XRCreative/Public/XRCreativeITFComponent.h
+++ b/Engine/Plugins/Experimental/XRCreativeFramework/Source/XRCreative/Public/XRCreativeITFComponent.h
@@ -19,6 +19,9 @@ class UXRCreativeITFRenderComponent;
 class UXRCreativePointerComponent;
 
 
+DECLARE_DYNAMIC_DELEGATE_RetVal_OneParam(bool, FCanSelectActorPredicate, AActor*, SelectionCandidate);
+
+
 UCLASS()
 class XRCREATIVE_API UXRCreativeITFComponent : public UActorComponent
 {
@@ -70,10 +73,23 @@ public:
 	UFUNCTION(BlueprintCallable, Category="XR Creative|Tools")
 	void SetCurrentCoordinateSystem(EToolContextCoordinateSystem CoordSystem);
 
+
+	UFUNCTION(BlueprintCallable, Category="XR Creative|Tools")
+	EToolContextTransformGizmoMode GetCurrentTransformGizmoMode() const { return CurrentTransformGizmoMode; }
+
+	UFUNCTION(BlueprintCallable, Category="XR Creative|Tools")
+	void SetCurrentTransformGizmoMode(EToolContextTransformGizmoMode GizmoMode);
+
 protected:
 	UPROPERTY(EditAnywhere, Category="XR Creative")
 	TSubclassOf<AXRCreativeBaseTransformGizmoActor> FullTRSGizmoActorClass;
 
+	UPROPERTY(EditAnywhere, BlueprintReadWrite, Category="XR Creative")
+	TArray<TSubclassOf<AActor>> UnselectableActorClasses;
+
+	UPROPERTY(BlueprintReadWrite, Category="XR Creative|Tools")
+	FCanSelectActorPredicate CanSelectPredicate;
+
 	UPROPERTY()
 	TObjectPtr<UXRCreativePointerComponent> PointerComponent;
 
@@ -95,6 +111,9 @@ protected:
 	UPROPERTY()
 	EToolContextCoordinateSystem CurrentCoordinateSystem = EToolContextCoordinateSystem::World;
 
+	UPROPERTY()
+	EToolContextTransformGizmoMode CurrentTransformGizmoMode = EToolContextTransformGizmoMode::Combined;
+
 protected:
 	void ToolsTick(float InDeltaTime);
 
diff --git a/Engine/Plugins/Experimental/XRCreativeFramework/Source/XRCreative/Public/XRCreativePointerComponent.h b/Engine/Plugins/Experimental/XRCreativeFramework/Source/XRCreative/Public/XRCreativePointerComponent.h
--- a/Engine/Plugins/Experimental/XRCreativeFramework/Source/XRCreative/Public/XRCreativePointerComponent.h
+++ b/Engine/Plugins/Experimental/XRCreativeFramework/Source/XRCreative/Public/XRCreativePointerComponent.h
@@ -13,7 +13,7 @@ namespace UE::XRCreative
 }
 
 
-UCLASS()
+UCLASS(ClassGroup = ("XR Creative"), meta=(BlueprintSpawnableComponent))
 class XRCREATIVE_API UXRCreativePointerComponent : public USceneComponent
 {
 	GENERATED_BODY()
@@ -26,7 +26,7 @@ public:
 	virtual void TickComponent(float InDeltaTime, enum ELevelTick InTickType, FActorComponentTickFunction* InThisTickFunction) override;
 
 	/**
-	* If bScaledByImpact is false, this returns the raw input to the smoothing filter, `TraceMaxLength` units away the +X direction.
+	* If bScaledByImpact is false, this returns the raw input to the smoothing filter, `TraceMaxLength` units away in the +X direction.
 	* If bScaledByImpact is true, the magnitude is shortened to the length of the blocking hit (if any).
 	*/
 	UFUNCTION(BlueprintCallable, Category="XR Creative")
@@ -61,6 +61,12 @@ protected:
 	UPROPERTY(EditAnywhere, Getter=IsEnabled, Setter=SetEnabled, Category="XR Creative")
 	bool bEnabled;
 
+	UPROPERTY(EditAnywhere, BlueprintReadWrite, Category="XR Creative")
+	TArray<TObjectPtr<AActor>> IgnoredActors;
+
+	UPROPERTY(EditAnywhere, BlueprintReadWrite, Category="XR Creative")
+	TArray<TObjectPtr<UPrimitiveComponent>> IgnoredComponents;
+
 protected:
 	FVector RawTraceEnd;
 	FHitResult HitResult;
diff --git a/Engine/Plugins/Experimental/XRCreativeFramework/Source/XRCreative/Public/XRCreativeSettings.h b/Engine/Plugins/Experimental/XRCreativeFramework/Source/XRCreative/Public/XRCreativeSettings.h
--- a/Engine/Plugins/Experimental/XRCreativeFramework/Source/XRCreative/Public/XRCreativeSettings.h
+++ b/Engine/Plugins/Experimental/XRCreativeFramework/Source/XRCreative/Public/XRCreativeSettings.h
@@ -5,6 +5,12 @@
 #include "Engine/DeveloperSettingsBackedByCVars.h"
 #include "XRCreativeSettings.generated.h"
 
+UENUM(BlueprintType)
+enum class EXRCreativeHandedness : uint8
+{
+	Left	UMETA(DisplayName = "Left"),
+	Right	UMETA(DisplayName = "Right"),
+};
 
 UCLASS(Config=XRCreativeSettings, DisplayName="XR Creative")
 class XRCREATIVE_API UXRCreativeSettings : public UDeveloperSettingsBackedByCVars
@@ -19,8 +25,11 @@ public:
 	TArray<float> FloatArray = {1,3,2,5};
 
 	UPROPERTY(Config, EditAnywhere, BlueprintReadWrite, Category="XR Creative")
-	TSoftClassPtr<class UXRCreativeSubsystemHelpers> SubsystemHelpersClass;
+	TSoftClassPtr<class UXRCreativeSubsystemHelper> SubsystemHelperClass;
+
+	UPROPERTY(Config, EditAnywhere, BlueprintReadWrite, Category="XR Creative", meta=(DisplayName="Handedness"))
+	EXRCreativeHandedness Handedness = EXRCreativeHandedness::Right;
 
 	UFUNCTION(BlueprintPure, Category="XR Creative")
 	static UXRCreativeSettings* GetXRCreativeSettings();
-};
+};
\ No newline at end of file
diff --git a/Engine/Plugins/Experimental/XRCreativeFramework/Source/XRCreative/Public/XRCreativeSubsystem.h b/Engine/Plugins/Experimental/XRCreativeFramework/Source/XRCreative/Public/XRCreativeSubsystem.h
--- a/Engine/Plugins/Experimental/XRCreativeFramework/Source/XRCreative/Public/XRCreativeSubsystem.h
+++ b/Engine/Plugins/Experimental/XRCreativeFramework/Source/XRCreative/Public/XRCreativeSubsystem.h
@@ -3,12 +3,14 @@
 #pragma once
 
 #include "Subsystems/EngineSubsystem.h"
-
 #include "XRCreativeSubsystem.generated.h"
 
 
+class UMVVMViewModelCollectionObject;
+
+
 UCLASS(Abstract, Blueprintable)
-class XRCREATIVE_API UXRCreativeSubsystemHelpers : public UObject
+class XRCREATIVE_API UXRCreativeSubsystemHelper : public UObject
 {
 	GENERATED_BODY()
 };
@@ -22,6 +24,12 @@ class XRCREATIVE_API UXRCreativeSubsystem : public UEngineSubsystem
 public:
 	virtual void Initialize(FSubsystemCollectionBase& Collection) override;
 
+	UFUNCTION(BlueprintCallable, Category = "XR Creative|Viewmodel")
+	UMVVMViewModelCollectionObject* GetViewModelCollection() const
+	{
+		return ViewModelCollection;
+	}
+
 #if WITH_EDITOR
 	/** Enter VR Mode */
 	UFUNCTION(BlueprintCallable, Category="XR Creative")
@@ -32,12 +40,18 @@ public:
 	static void ExitVRMode();
 #endif // #if WITH_EDITOR
 
+private:
+	UPROPERTY(Transient)
+	TObjectPtr<UMVVMViewModelCollectionObject> ViewModelCollection;
+
 protected:
 	void OnEngineInitComplete();
 
 protected:
 	UPROPERTY(BlueprintReadOnly, Category="XR Creative")
-	TObjectPtr<UXRCreativeSubsystemHelpers> Helpers;
+	TObjectPtr<UXRCreativeSubsystemHelper> Helpers;
 
 	FDelegateHandle EngineInitCompleteDelegate;
+
+
 };
diff --git a/Engine/Plugins/Experimental/XRCreativeFramework/Source/XRCreative/Public/XRCreativeTeleporter.h b/Engine/Plugins/Experimental/XRCreativeFramework/Source/XRCreative/Public/XRCreativeTeleporter.h
--- a/Engine/Plugins/Experimental/XRCreativeFramework/Source/XRCreative/Public/XRCreativeTeleporter.h
+++ b/Engine/Plugins/Experimental/XRCreativeFramework/Source/XRCreative/Public/XRCreativeTeleporter.h
@@ -5,6 +5,7 @@
 #include "GameFramework/Actor.h"
 #include "XRCreativeTeleporter.generated.h"
 
+class UXRDeviceVisualizationComponent;
 class UMotionControllerComponent;
 class UStaticMeshComponent;
 
@@ -39,4 +40,21 @@ protected:
 
 	UPROPERTY(EditAnywhere, BlueprintReadOnly, Category="XR Creative Teleporter")
 	TObjectPtr<UStaticMeshComponent> TeleportMesh;
+
+	UPROPERTY(EditAnywhere, BlueprintReadOnly, Category="XR Creative Teleporter")
+	TObjectPtr<UMotionControllerComponent> HeadMountedDisplay;
+
+	UPROPERTY(EditAnywhere, BlueprintReadOnly, Category="XR Creative Teleporter")
+	TObjectPtr<UStaticMeshComponent> LeftControllerVisual;
+
+	UPROPERTY(EditAnywhere, BlueprintReadOnly, Category="XR Creative Teleporter")
+	TObjectPtr<UStaticMeshComponent> RightControllerVisual;
+
+	UPROPERTY(EditAnywhere, BlueprintReadOnly, Category="XR Creative Teleporter")
+	TObjectPtr<UStaticMeshComponent> HMDVisual;
+
+	UPROPERTY(EditAnywhere, BlueprintReadOnly, Category="XR Creative Teleporter")
+	TObjectPtr<UXRDeviceVisualizationComponent> TestVisual;
+
+	
 };
diff --git a/Engine/Plugins/Experimental/XRCreativeFramework/Source/XRCreative/Public/XRCreativeToolActor.h b/Engine/Plugins/Experimental/XRCreativeFramework/Source/XRCreative/Public/XRCreativeToolActor.h
new file mode 100644
--- /dev/null
+++ b/Engine/Plugins/Experimental/XRCreativeFramework/Source/XRCreative/Public/XRCreativeToolActor.h
@@ -0,0 +1,41 @@
+// Copyright Epic Games, Inc. All Rights Reserved.
+
+#pragma once
+
+#include "CoreMinimal.h"
+#include "GameFramework/Actor.h"
+#include "XRCreativeToolActor.generated.h"
+
+class AXRCreativeAvatar;
+
+UCLASS(Blueprintable, Abstract)
+class XRCREATIVE_API AXRCreativeToolActor : public AActor
+{
+	GENERATED_BODY()
+
+public:
+	// Sets default values for this actor's properties
+	AXRCreativeToolActor();
+
+	UPROPERTY(BlueprintReadWrite, Category="XR Creative")
+	TObjectPtr<AXRCreativeAvatar> OwnerAvatar;
+
+protected:
+	// Called when the game starts or when spawned
+	virtual void BeginPlay() override;
+
+public:
+	// Called every frame
+	virtual void Tick(float DeltaTime) override;
+
+	//Called manually by the Avatar that spawns them. 
+
+	UFUNCTION(BlueprintImplementableEvent, BlueprintCallable,Category="XR Creative")
+	void InitializeTool();
+
+	UFUNCTION(BlueprintImplementableEvent, BlueprintCallable, Category="XR Creative")
+	void TickTool(float DeltaSeconds);
+	
+	UFUNCTION(BlueprintImplementableEvent, BlueprintCallable, Category="XR Creative")
+	bool ShutDownTool();
+};
diff --git a/Engine/Plugins/Experimental/XRCreativeFramework/Source/XRCreative/Public/XRCreativeToolset.h b/Engine/Plugins/Experimental/XRCreativeFramework/Source/XRCreative/Public/XRCreativeToolset.h
--- a/Engine/Plugins/Experimental/XRCreativeFramework/Source/XRCreative/Public/XRCreativeToolset.h
+++ b/Engine/Plugins/Experimental/XRCreativeFramework/Source/XRCreative/Public/XRCreativeToolset.h
@@ -7,10 +7,13 @@
 
 #include "XRCreativeToolset.generated.h"
 
+DECLARE_LOG_CATEGORY_EXTERN(LogXRCreativeToolset, Log, All);
 
+class UMVVMViewModelBase;
 class UCommonButtonStyle;
 class UCommonTextStyle;
 class UInputMappingContext;
+class AXRCreativeToolActor;
 
 
 UCLASS()
@@ -32,7 +35,7 @@ protected:
 	UPROPERTY(BlueprintReadOnly, Category="XR Creative")
 	TArray<TObjectPtr<UXRCreativePaletteTab>> Tabs;
 
-	UPROPERTY(BlueprintReadOnly, Category="XR Creative")
+	UPROPERTY(BlueprintReadWrite, Category="XR Creative")
 	TObjectPtr<AXRCreativeAvatar> Owner;
 };
 
@@ -51,14 +54,16 @@ public:
 
 	UFUNCTION(BlueprintCallable, Category="XR Creative")
 	virtual FText GetDisplayName() const PURE_VIRTUAL( UXRCreativeTool::GetDisplayName, return FText::GetEmpty(); );
-
+	
 	UFUNCTION(BlueprintCallable, Category="XR Creative")
 	virtual TSubclassOf<class UXRCreativePaletteToolTab> GetPaletteTabClass() const PURE_VIRTUAL( UXRCreativeTool::GetPaletteTabClass, return nullptr; );
+
+
 };
 
 
 UCLASS(Abstract, Blueprintable)
-class UXRCreativeBlueprintableTool : public UXRCreativeTool
+class XRCREATIVE_API UXRCreativeBlueprintableTool : public UXRCreativeTool
 {
 	GENERATED_BODY()
 
@@ -67,6 +72,8 @@ public:
 	virtual FText GetDisplayName() const override { return DisplayName; }
 	virtual TSubclassOf<UXRCreativePaletteToolTab> GetPaletteTabClass() const override { return PaletteTabClass; }
 
+	UInputMappingContext* GetInputMappingContext() { return ToolInputMappingContext; }
+
 protected:
 	UPROPERTY(BlueprintReadOnly, EditDefaultsOnly, Category="XR Creative")
 	FName ToolName;
@@ -74,14 +81,19 @@ protected:
 	UPROPERTY(BlueprintReadOnly, EditDefaultsOnly, Category="XR Creative")
 	FText DisplayName;
 
+	UPROPERTY(BlueprintReadOnly, EditDefaultsOnly, Category="XR Creative")
+	TSubclassOf<UMVVMViewModelBase> ToolViewmodel;
+
 	UPROPERTY(BlueprintReadOnly, EditDefaultsOnly, Category="XR Creative")
 	TSubclassOf<UXRCreativePaletteToolTab> PaletteTabClass;
 
 	UPROPERTY(BlueprintReadWrite, EditAnywhere, Category="XR Creative")
-	TSubclassOf<AActor> ToolActor;
+	TSubclassOf<AXRCreativeToolActor> ToolActor;
 
 	UPROPERTY(BlueprintReadWrite, EditAnywhere, Category="XR Creative")
 	TObjectPtr<UInputMappingContext> ToolInputMappingContext;
+
+
 	
 };
 
@@ -142,9 +154,21 @@ class XRCREATIVE_API UXRCreativeToolset : public UPrimaryDataAsset
 	GENERATED_BODY()
 
 public:
+	// UXRCreativeToolset();
+	//
+	// ~UXRCreativeToolset();
+	
 	UPROPERTY(EditDefaultsOnly, BlueprintReadOnly, Category="XR Creative")
 	TSubclassOf<AXRCreativeAvatar> Avatar;
 
+	/** Default Input Mapping is used for Right-Handed users, or if no LeftInputMappingContext is provided. */
+	UPROPERTY(EditDefaultsOnly, BlueprintReadOnly, Category="XR Creative")
+	TObjectPtr<UInputMappingContext> DefaultInputMappingContext;
+	
+	/** If Handedness is selected in XRCreative Settings, uses this entry in place of Default/Right  */
+	UPROPERTY(EditDefaultsOnly, BlueprintReadOnly, Category="XR Creative")
+	TObjectPtr<UInputMappingContext> LeftInputMappingContext;
+	
 	UPROPERTY(EditDefaultsOnly, BlueprintReadOnly, Category="XR Creative")
 	TArray<FXRCreativeToolEntry> Tools;
 
@@ -153,4 +177,17 @@ public:
 
 	UPROPERTY(EditDefaultsOnly, BlueprintReadOnly, Category="XR Creative")
 	TSoftObjectPtr<UXRCreativeStyle> Style;
+
+	UPROPERTY(EditAnywhere, Category="XR Creative")
+	bool bEnableUIMenuActor = false;
+
+	UFUNCTION(BlueprintCallable, Category="XR Creative")
+	bool GetEnableUIMenuActor() const { return bEnableUIMenuActor; };
+
+	/*Enable for legacy or custom menu actors.*/
+	UPROPERTY(EditDefaultsOnly, BlueprintReadOnly, Category="XR Creative", meta=(EditCondition="bEnableUIMenuActor"))
+	TSubclassOf<AActor> UIMenuActor;
+	
+	UFUNCTION(BlueprintCallable, Category="XR Creative")
+	virtual TSubclassOf<AActor> GetUIMenuActor() const { return UIMenuActor; };
 };
diff --git a/Engine/Plugins/Experimental/XRCreativeFramework/Source/XRCreative/XRCreative.Build.cs b/Engine/Plugins/Experimental/XRCreativeFramework/Source/XRCreative/XRCreative.Build.cs
--- a/Engine/Plugins/Experimental/XRCreativeFramework/Source/XRCreative/XRCreative.Build.cs
+++ b/Engine/Plugins/Experimental/XRCreativeFramework/Source/XRCreative/XRCreative.Build.cs
@@ -7,6 +7,7 @@ public class XRCreative : ModuleRules
 	public XRCreative(ReadOnlyTargetRules Target) : base(Target)
 	{
 		PCHUsage = ModuleRules.PCHUsageMode.UseExplicitOrSharedPCHs;
+		bTreatAsEngineModule = true;
 
 		PublicDependencyModuleNames.AddRange(
 			new string[]
@@ -28,6 +29,7 @@ public class XRCreative : ModuleRules
 				"InputCore",
 				"InteractiveToolsFramework",
 				"LevelSequence",
+				"ModelViewViewModel",
 				"RenderCore", // TODO: Kill this? Only for FlushRenderingCommands
 				"Slate",
 				"SlateCore",
diff --git a/Engine/Plugins/Experimental/XRCreativeFramework/Source/XRCreativeEditor/Private/XRCreativeEditorUtilityToolActor.cpp b/Engine/Plugins/Experimental/XRCreativeFramework/Source/XRCreativeEditor/Private/XRCreativeEditorUtilityToolActor.cpp
new file mode 100644
--- /dev/null
+++ b/Engine/Plugins/Experimental/XRCreativeFramework/Source/XRCreativeEditor/Private/XRCreativeEditorUtilityToolActor.cpp
@@ -0,0 +1,70 @@
+// Copyright Epic Games, Inc. All Rights Reserved.
+
+#include "XRCreativeEditorUtilityToolActor.h"
+#include "GameFramework/InputSettings.h"
+#include "Engine/InputDelegateBinding.h"
+
+
+// Sets default values
+
+
+AXRCreativeEditorUtilityToolActor::AXRCreativeEditorUtilityToolActor(const FObjectInitializer& ObjectInitializer)
+//	: Super(ObjectInitializer)
+{
+	PrimaryActorTick.bCanEverTick = true;
+}
+
+void AXRCreativeEditorUtilityToolActor::OnConstruction(const FTransform& Transform)
+{
+	Super::OnConstruction(Transform);
+
+	CreateEditorInput();
+}
+
+void AXRCreativeEditorUtilityToolActor::CreateEditorInput()
+{
+	if (bReceivesEditorInput && !HasAnyFlags(RF_ClassDefaultObject) && !EditorOnlyInputComponent)
+	{
+		EditorOnlyInputComponent = NewObject<UInputComponent>(this, UInputSettings::GetDefaultInputComponentClass(), TEXT("EdUtilActor_InputComponent0"), RF_Transient);
+		UInputDelegateBinding::BindInputDelegatesWithSubojects(this, EditorOnlyInputComponent);
+	}
+}
+
+void AXRCreativeEditorUtilityToolActor::RemoveEditorInput()
+{
+	ensure(!bReceivesEditorInput);
+	
+	if (EditorOnlyInputComponent)
+	{
+		EditorOnlyInputComponent->DestroyComponent();
+	}
+	EditorOnlyInputComponent = nullptr;
+}
+
+#if WITH_EDITOR
+void AXRCreativeEditorUtilityToolActor::PostEditChangeProperty(FPropertyChangedEvent& PropertyChangedEvent)
+{
+	Super::PostEditChangeProperty(PropertyChangedEvent);
+
+	const FName PropertyName = (PropertyChangedEvent.Property != nullptr) ? PropertyChangedEvent.Property->GetFName() : NAME_None;
+	if (PropertyName == GET_MEMBER_NAME_CHECKED(AXRCreativeEditorUtilityToolActor, bReceivesEditorInput))
+	{
+		SetReceivesEditorInput(bReceivesEditorInput);
+	}
+}
+#endif
+
+
+void AXRCreativeEditorUtilityToolActor::SetReceivesEditorInput(bool bInValue)
+{
+	bReceivesEditorInput = bInValue;
+	if (bReceivesEditorInput)
+	{
+		CreateEditorInput();
+	}
+	else
+	{
+		RemoveEditorInput();
+	}
+}
+
diff --git a/Engine/Plugins/Experimental/XRCreativeFramework/Source/XRCreativeEditor/Private/XRCreativeVREditorMode.cpp b/Engine/Plugins/Experimental/XRCreativeFramework/Source/XRCreativeEditor/Private/XRCreativeVREditorMode.cpp
--- a/Engine/Plugins/Experimental/XRCreativeFramework/Source/XRCreativeEditor/Private/XRCreativeVREditorMode.cpp
+++ b/Engine/Plugins/Experimental/XRCreativeFramework/Source/XRCreativeEditor/Private/XRCreativeVREditorMode.cpp
@@ -39,14 +39,15 @@ void UXRCreativeVREditorMode::Enter()
 		return;
 	}
 
-	Super::Enter();
-
 	UXRCreativeToolset* Toolset = ToolsetClass.LoadSynchronous();
 	Avatar = Cast<AXRCreativeAvatar>(SpawnTransientSceneActor(Toolset->Avatar, "XRCreativeAvatar"));
 	Avatar->RegisterObjectForInput(this);
 	Avatar->ConfigureToolset(Toolset);
 
+	Super::Enter();
+
 	BP_OnEnter();
+	Avatar->BP_OnVRInitialize();
 }
 
 
@@ -217,15 +218,37 @@ void UXRCreativeVREditorMode::EnableStereo()
 		Viewport->EnableStereoRendering(true);
 		Viewport->SetRenderDirectlyToWindow(true);
 	}
-
-	PRAGMA_DISABLE_DEPRECATION_WARNINGS;
-	TObjectPtr<UPlayerMappableInputConfig> InputConfig = Cast<UPlayerMappableInputConfig>(MappableInputConfig.TryLoad());
+	
 	TSet<UInputMappingContext*> Contexts;
-	for (const auto& Context : InputConfig->GetMappingContexts())
+	for (const TObjectPtr<UXRCreativeTool>& Tool : Avatar->GetTools())
+	{
+		if (UXRCreativeBlueprintableTool* BPTool = Cast<UXRCreativeBlueprintableTool>(Tool))
+		{
+			if (UInputMappingContext* ToolIMC = BPTool->GetInputMappingContext())
+			{
+				Contexts.Add(ToolIMC);
+			}
+		}
+	}
+
+	if (const UXRCreativeToolset* Toolset = Avatar->GetToolset())
 	{
-		Contexts.Add(Context.Key);
+		if (Toolset->DefaultInputMappingContext)
+		{
+			Contexts.Add(Toolset->DefaultInputMappingContext);
+		}
+
+		if (Toolset->LeftInputMappingContext)
+		{
+			Contexts.Add(Toolset->LeftInputMappingContext);
+		}
 	}
-	PRAGMA_ENABLE_DEPRECATION_WARNINGS;
+	
+	if (Contexts.IsEmpty())
+	{
+		UE_LOG(LogXRCreativeEditor, Warning, TEXT("No UInputMappingContexts provided in the UXRCreativeToolset. Action bindings will not be visible to the OpenXR runtime."));
+	}
+
 	UOpenXRInputFunctionLibrary::BeginXRSession(Contexts);
 }
 
diff --git a/Engine/Plugins/Experimental/XRCreativeFramework/Source/XRCreativeEditor/Private/XRCreativeVREditorMode.h b/Engine/Plugins/Experimental/XRCreativeFramework/Source/XRCreativeEditor/Private/XRCreativeVREditorMode.h
--- a/Engine/Plugins/Experimental/XRCreativeFramework/Source/XRCreativeEditor/Private/XRCreativeVREditorMode.h
+++ b/Engine/Plugins/Experimental/XRCreativeFramework/Source/XRCreativeEditor/Private/XRCreativeVREditorMode.h
@@ -70,8 +70,6 @@ protected:
 	void InvalidSettingNotification(const FText& ErrorDetails);
 
 protected:
-	UPROPERTY(EditDefaultsOnly, Category="XR Creative", meta=(AllowedClasses="/Script/EnhancedInput.PlayerMappableInputConfig"))
-	FSoftObjectPath MappableInputConfig;
 
 	UPROPERTY(EditDefaultsOnly, Category="XR Creative")
 	TSoftObjectPtr<UXRCreativeToolset> ToolsetClass;
diff --git a/Engine/Plugins/Experimental/XRCreativeFramework/Source/XRCreativeEditor/Public/XRCreativeEditorUtilityToolActor.h b/Engine/Plugins/Experimental/XRCreativeFramework/Source/XRCreativeEditor/Public/XRCreativeEditorUtilityToolActor.h
new file mode 100644
--- /dev/null
+++ b/Engine/Plugins/Experimental/XRCreativeFramework/Source/XRCreativeEditor/Public/XRCreativeEditorUtilityToolActor.h
@@ -0,0 +1,71 @@
+// Copyright Epic Games, Inc. All Rights Reserved.
+
+
+#pragma once
+
+#include "Components/InputComponent.h"
+#include "CoreMinimal.h"
+#include "GameFramework/Actor.h"
+#include "Math/Transform.h"
+#include "UObject/ObjectMacros.h"
+#include "UObject/ObjectPtr.h"
+#include "UObject/ScriptMacros.h"
+#include "UObject/UObjectGlobals.h"
+#include "CoreMinimal.h"
+#include "XRCreativeToolActor.h"
+#include "XRCreativeEditorUtilityToolActor.generated.h"
+
+
+
+class UInputComponent;
+class UObject;
+struct FFrame;
+struct FPropertyChangedEvent;
+
+
+
+UCLASS(Blueprintable, Abstract, meta = (ShowWorldContextPin))
+class XRCREATIVEEDITOR_API AXRCreativeEditorUtilityToolActor : public AXRCreativeToolActor
+{
+	GENERATED_UCLASS_BODY()
+
+	// Standard function to execute
+	UFUNCTION(BlueprintCallable, BlueprintImplementableEvent, Category = "Editor")
+	void Run();
+
+	virtual void OnConstruction(const FTransform& Transform) override;
+#if WITH_EDITOR
+	virtual void PostEditChangeProperty(FPropertyChangedEvent& PropertyChangedEvent) override;
+#endif
+
+	/** Returns the current InputComponent on this utility actor. This will be NULL unless bReceivesEditorInput is set to true. */
+	UFUNCTION(BlueprintCallable, BlueprintPure, Category = "Input|Editor")
+	UInputComponent* GetInputComponent() const
+	{ 
+		return EditorOnlyInputComponent.Get();
+	}
+	
+	UFUNCTION(BlueprintSetter, Category = "Input|Editor")
+	void SetReceivesEditorInput(bool bInValue);
+	
+	UFUNCTION(BlueprintGetter, BlueprintPure, Category = "Input|Editor")
+	bool GetReceivesEditorInput() const
+	{
+		return bReceivesEditorInput;
+	}	
+	
+private:
+
+	/** Creates the EditorOnlyInputComponent if it does not already exist and registers all subobject callbacks to it */
+	void CreateEditorInput();
+
+	/** Removes the EditorOnlyInputComponent from this utility actor */
+	void RemoveEditorInput();
+	
+	UPROPERTY(Transient, DuplicateTransient)
+	TObjectPtr<UInputComponent> EditorOnlyInputComponent;
+	
+	/** If set to true, then this actor will be able to recieve input delegate callbacks when in the editor. */
+	UPROPERTY(EditAnywhere, Category = "Input|Editor", BlueprintSetter = SetReceivesEditorInput, BlueprintGetter = GetReceivesEditorInput)
+	bool bReceivesEditorInput = false;
+};
\ No newline at end of file
diff --git a/Engine/Plugins/Experimental/XRCreativeFramework/Source/XRCreativeEditor/XRCreativeEditor.Build.cs b/Engine/Plugins/Experimental/XRCreativeFramework/Source/XRCreativeEditor/XRCreativeEditor.Build.cs
--- a/Engine/Plugins/Experimental/XRCreativeFramework/Source/XRCreativeEditor/XRCreativeEditor.Build.cs
+++ b/Engine/Plugins/Experimental/XRCreativeFramework/Source/XRCreativeEditor/XRCreativeEditor.Build.cs
@@ -7,6 +7,7 @@ public class XRCreativeEditor : ModuleRules
 	public XRCreativeEditor(ReadOnlyTargetRules Target) : base(Target)
 	{
 		PCHUsage = ModuleRules.PCHUsageMode.UseExplicitOrSharedPCHs;
+		bTreatAsEngineModule = true;
 
 		PublicDependencyModuleNames.AddRange(new string[]
 			{
diff --git a/Engine/Plugins/Experimental/XRCreativeFramework/XRCreativeFramework.uplugin b/Engine/Plugins/Experimental/XRCreativeFramework/XRCreativeFramework.uplugin
--- a/Engine/Plugins/Experimental/XRCreativeFramework/XRCreativeFramework.uplugin
+++ b/Engine/Plugins/Experimental/XRCreativeFramework/XRCreativeFramework.uplugin
@@ -10,14 +10,13 @@
 	"DocsURL": "",
 	"MarketplaceURL": "",
 	"SupportURL": "",
-	"CanContainContent": true,
+	"CanContainContent": false,
 	"IsBetaVersion": false,
 	"IsExperimentalVersion": true,
 	"Installed": false,
 	"EnabledByDefault": false,
 	"SupportedTargetPlatforms": [
-		"Win64",
-		"Linux"
+		"Win64"
 	],
 	"Modules": [
 		{
@@ -87,6 +86,10 @@
 		{
 			"Name": "MultiUserClient",
 			"Enabled": true
+		},
+		{
+			"Name": "ModelViewViewModel",
+			"Enabled":  true
 		}
 	]
 }
