XRCreative: Fix static analysis warning.

#rnx

diff --git a/Engine/Plugins/Experimental/XRCreativeFramework/Source/XRCreative/Private/ITF/XRCreativeITFComponent.cpp b/Engine/Plugins/Experimental/XRCreativeFramework/Source/XRCreative/Private/ITF/XRCreativeITFComponent.cpp
--- a/Engine/Plugins/Experimental/XRCreativeFramework/Source/XRCreative/Private/ITF/XRCreativeITFComponent.cpp
+++ b/Engine/Plugins/Experimental/XRCreativeFramework/Source/XRCreative/Private/ITF/XRCreativeITFComponent.cpp
@@ -617,7 +617,7 @@ void UXRCreativeITFComponent::ToolsTick(float InDeltaTime)
 	ULocalPlayer* LocalPlayer = Cast<ULocalPlayer>(PC->Player);
 	FVector ViewLocation;
 	FRotator ViewRotation;
-	const int32 StereoViewIndex = GEngine->IsStereoscopic3D(Viewport) ? EStereoscopicEye::eSSE_LEFT_EYE : INDEX_NONE;
+	const int32 StereoViewIndex = GEngine->IsStereoscopic3D(Viewport) ? eSSE_LEFT_EYE : eSSE_MONOSCOPIC;
 	FSceneView* SceneView = LocalPlayer->CalcSceneView(&ViewFamily, ViewLocation, ViewRotation,
 		LocalPlayer->ViewportClient->Viewport, nullptr, StereoViewIndex);
 	if (!SceneView)
