Virtual Scouting Incremental Checkin - Settings Systems

- Register settings for Virtual Scouting and XR Creative plugins.
- Update VScoutAvatar to use XRCreative Editor Settings where needed.

#virtualized

diff --git a/Engine/Plugins/Experimental/VirtualScouting/Source/VirtualScouting/Private/Module.cpp b/Engine/Plugins/Experimental/VirtualScouting/Source/VirtualScouting/Private/Module.cpp
--- a/Engine/Plugins/Experimental/VirtualScouting/Source/VirtualScouting/Private/Module.cpp
+++ b/Engine/Plugins/Experimental/VirtualScouting/Source/VirtualScouting/Private/Module.cpp
@@ -15,13 +15,26 @@ DEFINE_LOG_CATEGORY(LogVirtualScouting);
 
 void FVirtualScoutingModule::StartupModule()
 {
-	//RegisterSettings();
+	if(ISettingsModule* SettingsModule = FModuleManager::GetModulePtr<ISettingsModule>("Settings"))
+	{
+		SettingsModule->RegisterSettings("Project", "Plugins", "VirtualScoutingSettings",
+			LOCTEXT("RuntimeSettingsName", "Virtual Scouting Settings"), LOCTEXT("RuntimeSettingsDescription", "Virtual Scouting Settings"),
+			GetMutableDefault<UVirtualScoutingSettings>());
+
+		SettingsModule->RegisterSettings("Editor", "Plugins", "VirtualScoutingEditorSettings",
+			LOCTEXT("EditorSettingsName", "Virtual Scouting Editor Settings"), LOCTEXT("EditorSettingsDescription", "Virtual Scouting Editor Settings"),
+			GetMutableDefault<UVirtualScoutingEditorSettings>());
+	}
 }
 
 
 void FVirtualScoutingModule::ShutdownModule()
 {
-	//UnregisterSettings();
+	if(ISettingsModule* SettingsModule = FModuleManager::GetModulePtr<ISettingsModule>("Settings"))
+	{
+		SettingsModule->UnregisterSettings("Project", "Plugins", "VirtualScoutingSettings");
+		SettingsModule->UnregisterSettings("Editor", "Plugins", "VirtualScoutingEditorSettings");
+	}
 }
 
 
diff --git a/Engine/Plugins/Experimental/VirtualScouting/Source/VirtualScouting/Private/VirtualScoutingSettings.cpp b/Engine/Plugins/Experimental/VirtualScouting/Source/VirtualScouting/Private/VirtualScoutingSettings.cpp
--- a/Engine/Plugins/Experimental/VirtualScouting/Source/VirtualScouting/Private/VirtualScoutingSettings.cpp
+++ b/Engine/Plugins/Experimental/VirtualScouting/Source/VirtualScouting/Private/VirtualScoutingSettings.cpp
@@ -7,3 +7,8 @@ UVirtualScoutingSettings* UVirtualScoutingSettings::GetVirtualScoutingSettings()
 {
 	return GetMutableDefault<UVirtualScoutingSettings>();
 }
+
+UVirtualScoutingEditorSettings* UVirtualScoutingEditorSettings::GetVirtualScoutingEditorSettings()
+{
+	return GetMutableDefault<UVirtualScoutingEditorSettings>();
+}
diff --git a/Engine/Plugins/Experimental/VirtualScouting/Source/VirtualScouting/Public/VirtualScoutingSettings.h b/Engine/Plugins/Experimental/VirtualScouting/Source/VirtualScouting/Public/VirtualScoutingSettings.h
--- a/Engine/Plugins/Experimental/VirtualScouting/Source/VirtualScouting/Public/VirtualScoutingSettings.h
+++ b/Engine/Plugins/Experimental/VirtualScouting/Source/VirtualScouting/Public/VirtualScoutingSettings.h
@@ -7,14 +7,15 @@
 
 
 /**
- * Virtual Scouting Settings 
+ * Per project settings for Virtual Scouting.
  */
-UCLASS(Config=VirtualScoutingSettings, DisplayName="Virtual Scouting")
-class VIRTUALSCOUTING_API UVirtualScoutingSettings : public UDeveloperSettingsBackedByCVars
+UCLASS(Config=VirtualScoutingSettings, DefaultConfig, DisplayName="Virtual Scouting")
+class VIRTUALSCOUTING_API UVirtualScoutingSettings : public UObject
 {
 	GENERATED_BODY()
 	
 public:
+
 	UPROPERTY(Config, EditAnywhere, BlueprintReadWrite, Category="Virtual Scouting", meta=(DisplayName="Show Measurements in Imperial Units"))
 	bool bUseImperial = false;
 
@@ -24,3 +25,17 @@ public:
 	UFUNCTION(BlueprintPure, Category="Virtual Scouting")
 	static UVirtualScoutingSettings* GetVirtualScoutingSettings();
 };
+
+/**
+ * Per user settings for Virtual Scouting Editor.
+ */
+UCLASS(Config=EditorPerProjectUserSettings, DisplayName="Virtual Scouting Editor")
+class VIRTUALSCOUTING_API UVirtualScoutingEditorSettings : public UObject
+{
+	GENERATED_BODY()
+	
+public:
+
+	UFUNCTION(BlueprintPure, Category="Virtual Scouting Editor")
+	static UVirtualScoutingEditorSettings* GetVirtualScoutingEditorSettings();
+};
diff --git a/Engine/Plugins/Experimental/XRCreativeFramework/Source/XRCreative/Private/Module.cpp b/Engine/Plugins/Experimental/XRCreativeFramework/Source/XRCreative/Private/Module.cpp
--- a/Engine/Plugins/Experimental/XRCreativeFramework/Source/XRCreative/Private/Module.cpp
+++ b/Engine/Plugins/Experimental/XRCreativeFramework/Source/XRCreative/Private/Module.cpp
@@ -1,6 +1,8 @@
 // Copyright Epic Games, Inc. All Rights Reserved.
 
 #include "Module.h"
+#include "XRCreativeSettings.h"
+#include "ISettingsModule.h"
 #include "XRCreativeLog.h"
 
 
@@ -12,11 +14,27 @@ DEFINE_LOG_CATEGORY(LogXRCreative);
 
 void FXRCreativeModule::StartupModule()
 {
+	if(ISettingsModule* SettingsModule = FModuleManager::GetModulePtr<ISettingsModule>("Settings"))
+	{
+		SettingsModule->RegisterSettings("Project", "Plugins", "XRCreativeSettings",
+			LOCTEXT("RuntimeSettingsName", "XRCreative Settings"), LOCTEXT("RuntimeSettingsDescription", "XRCreative Settings"),
+			GetMutableDefault<UXRCreativeSettings>());
+
+		SettingsModule->RegisterSettings("Editor", "Plugins", "XRCreativeEditorSettings",
+			LOCTEXT("EditorSettingsName", "XRCreative Editor Settings"), LOCTEXT("EditorSettingsDescription", "XRCreative Editor Settings"),
+			GetMutableDefault<UXRCreativeEditorSettings>());
+	}
 }
 
 
 void FXRCreativeModule::ShutdownModule()
 {
+
+	if(ISettingsModule* SettingsModule = FModuleManager::GetModulePtr<ISettingsModule>("Settings"))
+	{
+		SettingsModule->UnregisterSettings("Project", "Plugins", "XRCreativeSettings");
+		SettingsModule->UnregisterSettings("Editor", "Plugins", "XRCreativeEditorSettings");
+	}
 }
 
 
diff --git a/Engine/Plugins/Experimental/XRCreativeFramework/Source/XRCreative/Private/XRCreativeSettings.cpp b/Engine/Plugins/Experimental/XRCreativeFramework/Source/XRCreative/Private/XRCreativeSettings.cpp
--- a/Engine/Plugins/Experimental/XRCreativeFramework/Source/XRCreative/Private/XRCreativeSettings.cpp
+++ b/Engine/Plugins/Experimental/XRCreativeFramework/Source/XRCreative/Private/XRCreativeSettings.cpp
@@ -7,3 +7,8 @@ UXRCreativeSettings* UXRCreativeSettings::GetXRCreativeSettings()
 {
 	return GetMutableDefault<UXRCreativeSettings>();
 }
+
+UXRCreativeEditorSettings* UXRCreativeEditorSettings::GetXRCreativeEditorSettings()
+{
+	return GetMutableDefault<UXRCreativeEditorSettings>();
+}
diff --git a/Engine/Plugins/Experimental/XRCreativeFramework/Source/XRCreative/Public/XRCreativeSettings.h b/Engine/Plugins/Experimental/XRCreativeFramework/Source/XRCreative/Public/XRCreativeSettings.h
--- a/Engine/Plugins/Experimental/XRCreativeFramework/Source/XRCreative/Public/XRCreativeSettings.h
+++ b/Engine/Plugins/Experimental/XRCreativeFramework/Source/XRCreative/Public/XRCreativeSettings.h
@@ -2,7 +2,6 @@
 
 #pragma once
 
-#include "Engine/DeveloperSettingsBackedByCVars.h"
 #include "XRCreativeSettings.generated.h"
 
 UENUM(BlueprintType)
@@ -12,24 +11,41 @@ enum class EXRCreativeHandedness : uint8
 	Right	UMETA(DisplayName = "Right"),
 };
 
-UCLASS(Config=XRCreativeSettings, DisplayName="XR Creative")
-class XRCREATIVE_API UXRCreativeSettings : public UDeveloperSettingsBackedByCVars
+/**
+ * Per project settings for XRCreative.
+ */
+UCLASS(Config=XRCreativeSettings, DefaultConfig, DisplayName="XR Creative")
+
+class XRCREATIVE_API UXRCreativeSettings : public UObject
 {
 	GENERATED_BODY()
 	
 public:
-	UPROPERTY(Config, EditAnywhere, BlueprintReadWrite, Category="XR Creative", meta=(DisplayName="Show Measurements in Imperial Units"))
-	bool bUseImperial = false;
-
-	UPROPERTY(Config, EditAnywhere, BlueprintReadWrite, Category="XR Creative", meta=(DisplayName="TestArray"))
-	TArray<float> FloatArray = {1,3,2,5};
 
 	UPROPERTY(Config, EditAnywhere, BlueprintReadWrite, Category="XR Creative")
 	TSoftClassPtr<class UXRCreativeSubsystemHelper> SubsystemHelperClass;
 
-	UPROPERTY(Config, EditAnywhere, BlueprintReadWrite, Category="XR Creative", meta=(DisplayName="Handedness"))
-	EXRCreativeHandedness Handedness = EXRCreativeHandedness::Right;
-
 	UFUNCTION(BlueprintPure, Category="XR Creative")
 	static UXRCreativeSettings* GetXRCreativeSettings();
-};
\ No newline at end of file
+};
+
+/**
+ * Per user settings for XRCreative Editor.
+ */
+UCLASS(Config=EditorPerProjectUserSettings, meta=(DisplayName="XR Creative Editor"))
+
+class XRCREATIVE_API UXRCreativeEditorSettings : public UObject
+{
+	GENERATED_BODY()
+
+public:
+
+	/** Manages Left/Right handedness user preferences.
+	 * Modifying this setting requires an editor restart to take effect.
+	 **/
+	UPROPERTY(Config, EditAnywhere, BlueprintReadWrite, Category="XR Creative", meta=(DisplayName="Handedness"))
+	EXRCreativeHandedness Handedness = EXRCreativeHandedness::Right;
+	
+	UFUNCTION(BlueprintPure, Category="XR Creative Editor")
+	static UXRCreativeEditorSettings* GetXRCreativeEditorSettings();
+};
