XR Creative - made all components of XRCreative gizmo hidden from scene captures so gizmo is not visiable using the viewfinder tool

diff --git a/Engine/Plugins/Experimental/XRCreativeFramework/Source/XRCreative/Private/ITF/XRCreativeGizmos.cpp b/Engine/Plugins/Experimental/XRCreativeFramework/Source/XRCreative/Private/ITF/XRCreativeGizmos.cpp
--- a/Engine/Plugins/Experimental/XRCreativeFramework/Source/XRCreative/Private/ITF/XRCreativeGizmos.cpp
+++ b/Engine/Plugins/Experimental/XRCreativeFramework/Source/XRCreative/Private/ITF/XRCreativeGizmos.cpp
@@ -33,6 +33,7 @@ UGizmoArrowComponent* AXRCreativeBaseTransformGizmoActor::SetupAxisArrow(UGizmoA
 	InComponent->Gap = 20.0f;
 	InComponent->Thickness = GizmoLineThickness;
 	InComponent->NotifyExternalPropertyUpdates();
+	InComponent->bHiddenInSceneCapture=true;
 	return InComponent;
 }
 
@@ -47,6 +48,7 @@ UGizmoRectangleComponent* AXRCreativeBaseTransformGizmoActor::SetupPlaneRect(UGi
 	InComponent->SegmentFlags = 0x2 | 0x4;
 	InComponent->Thickness = GizmoLineThickness;
 	InComponent->NotifyExternalPropertyUpdates();
+	InComponent->bHiddenInSceneCapture=true;
 	return InComponent;
 }
 
@@ -59,6 +61,7 @@ UGizmoCircleComponent* AXRCreativeBaseTransformGizmoActor::SetupAxisRotateCircle
 	InComponent->Radius = 100.0f;
 	InComponent->Thickness = GizmoLineThickness;
 	InComponent->NotifyExternalPropertyUpdates();
+	InComponent->bHiddenInSceneCapture=true;
 	return InComponent;
 }
 
@@ -77,6 +80,7 @@ UGizmoRectangleComponent* AXRCreativeBaseTransformGizmoActor::SetupAxisScaleRect
 	InComponent->bOrientYAccordingToCamera = true;
 	InComponent->SegmentFlags = 0x1 | 0x2 | 0x4; // | 0x8;
 	InComponent->NotifyExternalPropertyUpdates();
+	InComponent->bHiddenInSceneCapture=true;
 	return InComponent;
 }
 
@@ -92,6 +96,7 @@ UGizmoRectangleComponent* AXRCreativeBaseTransformGizmoActor::SetupPlaneScaleFun
 	InComponent->Thickness = GizmoLineThickness;
 	InComponent->SegmentFlags = 0x2 | 0x4;
 	InComponent->NotifyExternalPropertyUpdates();
+	InComponent->bHiddenInSceneCapture=true;
 	return InComponent;
 }
 
@@ -151,6 +156,7 @@ void AXRCreativeBaseTransformGizmoActor::ConstructDefaults(ETransformGizmoSubEle
 		RotationSphereComp->Radius = 120.0f;
 		RotationSphereComp->bViewAligned = true;
 		RotationSphereComp->NotifyExternalPropertyUpdates();
+		RotationSphereComp->bHiddenInSceneCapture=true;
 		RotationSphere = RotationSphereComp;
 	}
 
@@ -163,6 +169,7 @@ void AXRCreativeBaseTransformGizmoActor::ConstructDefaults(ETransformGizmoSubEle
 		UniformScaleComp->Origin = FVector(BoxSize/2, BoxSize/2, BoxSize/2);
 		UniformScaleComp->Dimensions = FVector(BoxSize, BoxSize, BoxSize);
 		UniformScaleComp->NotifyExternalPropertyUpdates();
+		UniformScaleComp->bHiddenInSceneCapture=true;
 		UniformScale = UniformScaleComp;
 	}
 
