XRCreative - fix/workaround for OpenXRinput with pre-declared input maps not respecting input priority.
-added function to clear all input mappings to XRCreativeAvatar
-exposed input mapping options to Add and Remove input mapping functions in XRCreativeAvatar

diff --git a/Engine/Plugins/Experimental/XRCreativeFramework/Source/XRCreative/Private/XRCreativeAvatar.cpp b/Engine/Plugins/Experimental/XRCreativeFramework/Source/XRCreative/Private/XRCreativeAvatar.cpp
--- a/Engine/Plugins/Experimental/XRCreativeFramework/Source/XRCreative/Private/XRCreativeAvatar.cpp
+++ b/Engine/Plugins/Experimental/XRCreativeFramework/Source/XRCreative/Private/XRCreativeAvatar.cpp
@@ -272,29 +272,37 @@ void AXRCreativeAvatar::UnregisterObjectForInput(UObject* InObject)
 }
 
 
-void AXRCreativeAvatar::AddInputMappingContext(UInputMappingContext* InContext, int32 InPriority)
+void AXRCreativeAvatar::AddInputMappingContext(UInputMappingContext* InContext, int32 InPriority,const FModifyContextOptions Options)
 {
 	if (IEnhancedInputSubsystemInterface* EnhancedInputSubsystemInterface = GetEnhancedInputSubsystemInterface())
 	{
 		if (IsValid(InContext))
 		{
-			EnhancedInputSubsystemInterface->AddMappingContext(InContext, InPriority);
+			EnhancedInputSubsystemInterface->AddMappingContext(InContext, InPriority, Options);
 		}
 	}
 }
 
 
-void AXRCreativeAvatar::RemoveInputMappingContext(UInputMappingContext* InContext)
+void AXRCreativeAvatar::RemoveInputMappingContext(UInputMappingContext* InContext, const FModifyContextOptions Options)
 {
 	if (IEnhancedInputSubsystemInterface* EnhancedInputSubsystemInterface = GetEnhancedInputSubsystemInterface())
 	{
 		if (IsValid(InContext))
 		{
-			EnhancedInputSubsystemInterface->RemoveMappingContext(InContext);
+			EnhancedInputSubsystemInterface->RemoveMappingContext(InContext, Options);
 		}
 	}
 }
 
+void AXRCreativeAvatar::ClearAllInputMappings()
+{
+	if (IEnhancedInputSubsystemInterface* EnhancedInputSubsystemInterface = GetEnhancedInputSubsystemInterface())
+	{
+		EnhancedInputSubsystemInterface->ClearAllMappings();
+	}
+}
+
 //** Cues up a Haptic effect to be processed on tick in ProcessHaptics() //
 void AXRCreativeAvatar::PlayHapticEffect(UHapticFeedbackEffect_Base* HapticEffect, const int ControllerID, const EControllerHand Hand, float Scale, bool bLoop)
 {
diff --git a/Engine/Plugins/Experimental/XRCreativeFramework/Source/XRCreative/Public/XRCreativeAvatar.h b/Engine/Plugins/Experimental/XRCreativeFramework/Source/XRCreative/Public/XRCreativeAvatar.h
--- a/Engine/Plugins/Experimental/XRCreativeFramework/Source/XRCreative/Public/XRCreativeAvatar.h
+++ b/Engine/Plugins/Experimental/XRCreativeFramework/Source/XRCreative/Public/XRCreativeAvatar.h
@@ -6,6 +6,7 @@
 #include "XRCreativeAvatar.generated.h"
 
 
+
 #if WITH_EDITOR
 class IConcertClientSession;
 #endif
@@ -82,12 +83,14 @@ public:
 
 	// Adds an explicitly provided Input Mapping Context to the input system
 	UFUNCTION(BlueprintCallable, Category="XR Creative")
-	void AddInputMappingContext(UInputMappingContext* Context, int32 Priority);
+	void AddInputMappingContext(UInputMappingContext* Context, int32 Priority, const FModifyContextOptions Options);
 
 	// Removes an explicitly provided Input Mapping Context to the input system
 	UFUNCTION(BlueprintCallable, Category="XR Creative")
-	void RemoveInputMappingContext(UInputMappingContext* Context);
-
+	void RemoveInputMappingContext(UInputMappingContext* Context, const FModifyContextOptions Options);
+	
+	UFUNCTION(BlueprintCallable, Category="XR Creative")
+	void ClearAllInputMappings();
 	/**
 	* Called when In-Editor VR mode is started. In-Editor equivalent to Begin Play.
 	*/
diff --git a/Engine/Plugins/Experimental/XRCreativeFramework/Source/XRCreativeEditor/Private/XRCreativeVREditorMode.cpp b/Engine/Plugins/Experimental/XRCreativeFramework/Source/XRCreativeEditor/Private/XRCreativeVREditorMode.cpp
--- a/Engine/Plugins/Experimental/XRCreativeFramework/Source/XRCreativeEditor/Private/XRCreativeVREditorMode.cpp
+++ b/Engine/Plugins/Experimental/XRCreativeFramework/Source/XRCreativeEditor/Private/XRCreativeVREditorMode.cpp
@@ -249,6 +249,8 @@ void UXRCreativeVREditorMode::EnableStereo()
 		UE_LOG(LogXRCreativeEditor, Warning, TEXT("No UInputMappingContexts provided in the UXRCreativeToolset. Action bindings will not be visible to the OpenXR runtime."));
 	}
 
+	Contexts.Empty(); // TODO this is a temporary fix because OpenXRInput.cpp does not respect input priorities. So instead we force it to fall back to legacy bindings because these actually work. 
+
 	UOpenXRInputFunctionLibrary::BeginXRSession(Contexts);
 }
 
