XRCreative - added Name text field to toolset data asset

diff --git a/Engine/Plugins/Experimental/XRCreativeFramework/Source/XRCreative/Public/XRCreativeToolset.h b/Engine/Plugins/Experimental/XRCreativeFramework/Source/XRCreative/Public/XRCreativeToolset.h
--- a/Engine/Plugins/Experimental/XRCreativeFramework/Source/XRCreative/Public/XRCreativeToolset.h
+++ b/Engine/Plugins/Experimental/XRCreativeFramework/Source/XRCreative/Public/XRCreativeToolset.h
@@ -157,6 +157,9 @@ public:
 	// UXRCreativeToolset();
 	//
 	// ~UXRCreativeToolset();
+
+	UPROPERTY(EditDefaultsOnly, BlueprintReadOnly, Category="XR Creative")
+	FText ToolsetName;
 	
 	UPROPERTY(EditDefaultsOnly, BlueprintReadOnly, Category="XR Creative")
 	TSubclassOf<AXRCreativeAvatar> Avatar;
