Virtual Scouting Incremental Checkin - Handedness Sub-systems

- Avatar now has sub-systems for managing handedness
- Fixed hard-reference in Draw Tool to Right Laser Pointer

#virtualized

diff --git a/Engine/Plugins/Experimental/XRCreativeFramework/Source/XRCreative/Private/XRCreativeAvatar.cpp b/Engine/Plugins/Experimental/XRCreativeFramework/Source/XRCreative/Private/XRCreativeAvatar.cpp
--- a/Engine/Plugins/Experimental/XRCreativeFramework/Source/XRCreative/Private/XRCreativeAvatar.cpp
+++ b/Engine/Plugins/Experimental/XRCreativeFramework/Source/XRCreative/Private/XRCreativeAvatar.cpp
@@ -36,6 +36,7 @@
 #	include "IMultiUserClientModule.h"
 #	include "IVREditorModule.h"
 #	include "VREditorModeBase.h"
+#	include "XRCreativeSettings.h"
 #endif
 
 
@@ -93,12 +94,24 @@ AXRCreativeAvatar::AXRCreativeAvatar(const FObjectInitializer& ObjectInitializer
 	WidgetInteraction->bTickInEditor = true;
 	WidgetInteraction->bShowDebug = false;
 	WidgetInteraction->InteractionDistance = 50.0;
-
+	
 	ToolsComponent = CreateDefaultSubobject<UXRCreativeITFComponent>("ToolsComponent");
 	ToolsComponent->bTickInEditor = true;
 	ToolsComponent->SetPointerComponent(RightControllerPointer);
 
 	BaseEyeHeight = 0.0f;
+
+#if WITH_EDITOR
+	UXRCreativeEditorSettings* Settings = UXRCreativeEditorSettings::GetXRCreativeEditorSettings();
+	UE_LOG(LogXRCreative, Log, TEXT("Handedness: %s"), *UEnum::GetValueAsString(Settings->Handedness) );
+
+	if (Settings->Handedness == EXRCreativeHandedness::Left)
+	{
+		MenuWidget->SetupAttachment(RightControllerAim);
+		WidgetInteraction->SetupAttachment(LeftControllerAim);
+		ToolsComponent->SetPointerComponent(LeftControllerPointer);
+	}
+#endif
 }
 
 
