Virtual Scouting Incremental Checkin - Tool Actor Input Mapping Handedness Subsystem Integration

- Migrated tools to make use of Handedness BPIs.
- Authored default Left-handed version of Input Mapping Contexts for all tools.

#virtualized

diff --git a/Engine/Plugins/Experimental/XRCreativeFramework/Source/XRCreative/Private/XRCreativeToolset.cpp b/Engine/Plugins/Experimental/XRCreativeFramework/Source/XRCreative/Private/XRCreativeToolset.cpp
--- a/Engine/Plugins/Experimental/XRCreativeFramework/Source/XRCreative/Private/XRCreativeToolset.cpp
+++ b/Engine/Plugins/Experimental/XRCreativeFramework/Source/XRCreative/Private/XRCreativeToolset.cpp
@@ -1,23 +1,37 @@
 // Copyright Epic Games, Inc. All Rights Reserved.
 
 #include "XRCreativeToolset.h"
-/*#include "XRCreativeSettings.h"
+
+#include "InputMappingContext.h"
+#include "XRCreativeLog.h"
 #include "Logging/LogMacros.h"
 
+#if WITH_EDITOR
+#	include "XRCreativeSettings.h"
+#endif
+
 DEFINE_LOG_CATEGORY(LogXRCreativeToolset);
 
-UXRCreativeToolset::UXRCreativeToolset()
+UInputMappingContext* UXRCreativeBlueprintableTool::GetToolInputMappingContext()
 {
-	UXRCreativeSettings* Settings = UXRCreativeSettings::GetXRCreativeSettings();
-	if (Settings->Handedness == EXRCreativeHandedness::Left && !LeftInputMappingContext)
-	{
-		UE_LOG(LogXRCreativeToolset, Warning, TEXT("Handedness is Left but no Left Input Mapping Context found in Toolset."));
-	}
-}
+	UInputMappingContext* InputMappingContext = DefaultToolInputMappingContext;
+	
+	#if WITH_EDITOR
+		UXRCreativeEditorSettings* Settings = UXRCreativeEditorSettings::GetXRCreativeEditorSettings();
 
-UXRCreativeToolset::~UXRCreativeToolset()
-{
+		if (Settings->Handedness == EXRCreativeHandedness::Left && LeftToolInputMappingContext)
+		{
+			InputMappingContext = LeftToolInputMappingContext;
+		}
+		else if (Settings->Handedness == EXRCreativeHandedness::Left && !LeftToolInputMappingContext)
+		{
+			InputMappingContext = DefaultToolInputMappingContext;
+			UE_LOG(LogXRCreativeToolset, Warning, TEXT("Handedness is Left but no Left Input Mapping Context found in Toolset - Using Default."));
+		}
+	
+	#endif
 	
-}*/
+	return InputMappingContext; 
+}
 
 
diff --git a/Engine/Plugins/Experimental/XRCreativeFramework/Source/XRCreative/Public/XRCreativeToolset.h b/Engine/Plugins/Experimental/XRCreativeFramework/Source/XRCreative/Public/XRCreativeToolset.h
--- a/Engine/Plugins/Experimental/XRCreativeFramework/Source/XRCreative/Public/XRCreativeToolset.h
+++ b/Engine/Plugins/Experimental/XRCreativeFramework/Source/XRCreative/Public/XRCreativeToolset.h
@@ -72,7 +72,8 @@ public:
 	virtual FText GetDisplayName() const override { return DisplayName; }
 	virtual TSubclassOf<UXRCreativePaletteToolTab> GetPaletteTabClass() const override { return PaletteTabClass; }
 
-	UInputMappingContext* GetInputMappingContext() { return ToolInputMappingContext; }
+	UFUNCTION(BlueprintPure, Category="XR Creative")
+	UInputMappingContext* GetToolInputMappingContext();
 
 protected:
 	UPROPERTY(BlueprintReadOnly, EditDefaultsOnly, Category="XR Creative")
@@ -91,8 +92,10 @@ protected:
 	TSubclassOf<AXRCreativeToolActor> ToolActor;
 
 	UPROPERTY(BlueprintReadWrite, EditAnywhere, Category="XR Creative")
-	TObjectPtr<UInputMappingContext> ToolInputMappingContext;
+	TObjectPtr<UInputMappingContext> DefaultToolInputMappingContext;
 
+	UPROPERTY(BlueprintReadWrite, EditAnywhere, Category="XR Creative")
+	TObjectPtr<UInputMappingContext> LeftToolInputMappingContext;
 
 	
 };
diff --git a/Engine/Plugins/Experimental/XRCreativeFramework/Source/XRCreativeEditor/Private/XRCreativeVREditorMode.cpp b/Engine/Plugins/Experimental/XRCreativeFramework/Source/XRCreativeEditor/Private/XRCreativeVREditorMode.cpp
--- a/Engine/Plugins/Experimental/XRCreativeFramework/Source/XRCreativeEditor/Private/XRCreativeVREditorMode.cpp
+++ b/Engine/Plugins/Experimental/XRCreativeFramework/Source/XRCreativeEditor/Private/XRCreativeVREditorMode.cpp
@@ -224,7 +224,7 @@ void UXRCreativeVREditorMode::EnableStereo()
 	{
 		if (UXRCreativeBlueprintableTool* BPTool = Cast<UXRCreativeBlueprintableTool>(Tool))
 		{
-			if (UInputMappingContext* ToolIMC = BPTool->GetInputMappingContext())
+			if (UInputMappingContext* ToolIMC = BPTool->GetToolInputMappingContext())
 			{
 				Contexts.Add(ToolIMC);
 			}
