Virtual Scouting - changed parent type of palette menu to CommonActivatableWidget, from CommonUserWidget to make activation event available

diff --git a/Engine/Plugins/Experimental/XRCreativeFramework/Source/XRCreative/Public/XRCreativeToolset.h b/Engine/Plugins/Experimental/XRCreativeFramework/Source/XRCreative/Public/XRCreativeToolset.h
--- a/Engine/Plugins/Experimental/XRCreativeFramework/Source/XRCreative/Public/XRCreativeToolset.h
+++ b/Engine/Plugins/Experimental/XRCreativeFramework/Source/XRCreative/Public/XRCreativeToolset.h
@@ -24,7 +24,7 @@ class UXRCreativePaletteTab : public UCommonActivatableWidget
 
 
 UCLASS()
-class UXRCreativePalette : public UCommonUserWidget
+class UXRCreativePalette : public UCommonActivatableWidget
 {
 	GENERATED_BODY()
 
