Virtual Scouting - Removed XRCreative style data asset as not used in any implementation. Removed refs to XRCreativestyle in the Toolset data-asset

diff --git a/Engine/Plugins/Experimental/XRCreativeFramework/Source/XRCreative/Public/XRCreativeToolset.h b/Engine/Plugins/Experimental/XRCreativeFramework/Source/XRCreative/Public/XRCreativeToolset.h
--- a/Engine/Plugins/Experimental/XRCreativeFramework/Source/XRCreative/Public/XRCreativeToolset.h
+++ b/Engine/Plugins/Experimental/XRCreativeFramework/Source/XRCreative/Public/XRCreativeToolset.h
@@ -114,30 +114,6 @@ protected:
 
 //////////////////////////////////////////////////////////////////////////
 
-
-UCLASS(BlueprintType)
-class UXRCreativeStyle : public UPrimaryDataAsset
-{
-	GENERATED_BODY()
-
-public:
-	UPROPERTY(EditDefaultsOnly, BlueprintReadOnly, Category="XR Creative")
-	TSubclassOf<UCommonTextStyle> RegularTextStyle;
-
-	UPROPERTY(EditDefaultsOnly, BlueprintReadOnly, Category="XR Creative")
-	TSubclassOf<UCommonTextStyle> BoldTextStyle;
-
-	UPROPERTY(EditDefaultsOnly, BlueprintReadOnly, Category="XR Creative")
-	TSubclassOf<UCommonTextStyle> ItalicTextStyle;
-
-	UPROPERTY(EditDefaultsOnly, BlueprintReadOnly, Category="XR Creative")
-	TSubclassOf<UCommonButtonStyle> RegularButtonStyle;
-
-	UPROPERTY(EditDefaultsOnly, BlueprintReadOnly, Category="XR Creative")
-	TSubclassOf<UCommonButtonStyle> LargeButtonStyle;
-};
-
-
 USTRUCT(BlueprintType)
 struct FXRCreativeToolEntry
 {
@@ -181,9 +157,6 @@ public:
 	UPROPERTY(EditDefaultsOnly, BlueprintReadOnly, Category="XR Creative")
 	TSubclassOf<UXRCreativePalette> Palette;
 
-	UPROPERTY(EditDefaultsOnly, BlueprintReadOnly, Category="XR Creative")
-	TSoftObjectPtr<UXRCreativeStyle> Style;
-
 	UPROPERTY(EditAnywhere, Category="XR Creative")
 	bool bEnableUIMenuActor = false;
 
