XR Creative: Add blueprint OnUndo / OnRedo delegates to UXRCreativeITFComponent. Minor cleanup.

diff --git a/Engine/Plugins/Experimental/XRCreativeFramework/Source/XRCreative/Private/ITF/XRCreativeITFComponent.cpp b/Engine/Plugins/Experimental/XRCreativeFramework/Source/XRCreative/Private/ITF/XRCreativeITFComponent.cpp
--- a/Engine/Plugins/Experimental/XRCreativeFramework/Source/XRCreative/Private/ITF/XRCreativeITFComponent.cpp
+++ b/Engine/Plugins/Experimental/XRCreativeFramework/Source/XRCreative/Private/ITF/XRCreativeITFComponent.cpp
@@ -377,6 +377,17 @@ void UXRCreativeITFComponent::InitializeComponent()
 {
 	Super::InitializeComponent();
 
+#if WITH_EDITOR
+	if (GEditor)
+	{
+		if (UTransBuffer* TransBuffer = Cast<UTransBuffer>(GEditor->Trans))
+		{
+			TransBuffer->OnUndo().AddUObject(this, &UXRCreativeITFComponent::HandleTransactorUndo);
+			TransBuffer->OnRedo().AddUObject(this, &UXRCreativeITFComponent::HandleTransactorRedo);
+		}
+	}
+#endif
+
 	ensure(IsValid(PointerComponent));
 
 	// Ensure our motion controller components tick before we do
@@ -454,6 +465,17 @@ void UXRCreativeITFComponent::UninitializeComponent()
 {
 	bIsShuttingDown = true;
 
+#if WITH_EDITOR
+	if (GEditor)
+	{
+		if (UTransBuffer* TransBuffer = Cast<UTransBuffer>(GEditor->Trans))
+		{
+			TransBuffer->OnUndo().RemoveAll(this);
+			TransBuffer->OnRedo().RemoveAll(this);
+		}
+	}
+#endif
+
 	if (ToolsContext)
 	{
 		//CancelOrCompleteActiveTool();
@@ -821,7 +843,25 @@ void UXRCreativeITFComponent::EditorToolsTick(float InDeltaTime)
 	// force rendering flush so that PDI lines get drawn
 	FlushRenderingCommands();
 }
-#endif
+
+
+void UXRCreativeITFComponent::HandleTransactorUndo(const FTransactionContext& TransactionContext, bool bSucceeded)
+{
+	if (bSucceeded)
+	{
+		OnUndo.Broadcast();
+	}
+}
+
+
+void UXRCreativeITFComponent::HandleTransactorRedo(const FTransactionContext& TransactionContext, bool bSucceeded)
+{
+	if (bSucceeded)
+	{
+		OnRedo.Broadcast();
+	}
+}
+#endif // #if WITH_EDITOR
 
 
 void UXRCreativeITFComponent::LeftMousePressed()
diff --git a/Engine/Plugins/Experimental/XRCreativeFramework/Source/XRCreative/Public/XRCreativeITFComponent.h b/Engine/Plugins/Experimental/XRCreativeFramework/Source/XRCreative/Public/XRCreativeITFComponent.h
--- a/Engine/Plugins/Experimental/XRCreativeFramework/Source/XRCreative/Public/XRCreativeITFComponent.h
+++ b/Engine/Plugins/Experimental/XRCreativeFramework/Source/XRCreative/Public/XRCreativeITFComponent.h
@@ -18,6 +18,10 @@ class UXRCreativeTransformInteraction;
 class UXRCreativeITFRenderComponent;
 class UXRCreativePointerComponent;
 
+#if WITH_EDITOR
+struct FTransactionContext;
+#endif
+
 
 DECLARE_DYNAMIC_DELEGATE_RetVal_OneParam(bool, FCanSelectActorPredicate, AActor*, SelectionCandidate);
 
@@ -51,6 +55,14 @@ public:
 	UFUNCTION(BlueprintCallable, Category="XR Creative|Tools")
 	void Redo();
 
+	DECLARE_DYNAMIC_MULTICAST_DELEGATE(FOnUndoRedo);
+
+	UPROPERTY(BlueprintAssignable, Category="XR Creative|Tools")
+	FOnUndoRedo OnUndo;
+
+	UPROPERTY(BlueprintAssignable, Category="XR Creative|Tools")
+	FOnUndoRedo OnRedo;
+
 
 	UFUNCTION(BlueprintCallable, Category="XR Creative|Tools")
 	void LeftMousePressed();
@@ -119,6 +131,9 @@ protected:
 
 #if WITH_EDITOR
 	void EditorToolsTick(float InDeltaTime);
+
+	void HandleTransactorUndo(const FTransactionContext& TransactionContext, bool Succeeded);
+	void HandleTransactorRedo(const FTransactionContext& TransactionContext, bool Succeeded);
 #endif
 
 protected:
diff --git a/Engine/Plugins/Experimental/XRCreativeFramework/XRCreativeFramework.uplugin b/Engine/Plugins/Experimental/XRCreativeFramework/XRCreativeFramework.uplugin
--- a/Engine/Plugins/Experimental/XRCreativeFramework/XRCreativeFramework.uplugin
+++ b/Engine/Plugins/Experimental/XRCreativeFramework/XRCreativeFramework.uplugin
@@ -47,42 +47,6 @@
 			"Name": "CommonUI",
 			"Enabled": true
 		},
-		{
-			"Name": "BlueprintFileUtils",
-			"Enabled": true
-		},
-		{
-			"Name": "JsonBlueprintUtilities",
-			"Enabled": true
-		},
-		{
-			"Name": "SunPosition",
-			"Enabled": true
-		},
-		{
-			"Name": "GeometryScripting",
-			"Enabled": true
-		},
-		{
-			"Name": "GizmoFramework",
-			"Enabled": true
-		},
-		{
-			"Name": "GizmoEdMode",
-			"Enabled": true
-		},
-		{
-			"Name": "ModelingToolsEditorMode",
-			"Enabled": true
-		},
-		{
-			"Name": "ConsoleVariables",
-			"Enabled": true
-		},
-		{
-			"Name": "LevelSnapshots",
-			"Enabled": true
-		},
 		{
 			"Name": "MultiUserClient",
 			"Enabled": true
