Virtual Scouting - Tidied up Teleporter class. Removed redundant components, added controller mesh and HMD mesh to remove out-of-plugin asset refs.
Added material to fade teleport material based on distance and worldscale of the teleporter and it's sub-components.

#virtualized

diff --git a/Engine/Plugins/Experimental/XRCreativeFramework/Source/XRCreative/Private/XRCreativeTeleporter.cpp b/Engine/Plugins/Experimental/XRCreativeFramework/Source/XRCreative/Private/XRCreativeTeleporter.cpp
--- a/Engine/Plugins/Experimental/XRCreativeFramework/Source/XRCreative/Private/XRCreativeTeleporter.cpp
+++ b/Engine/Plugins/Experimental/XRCreativeFramework/Source/XRCreative/Private/XRCreativeTeleporter.cpp
@@ -2,7 +2,6 @@
 
 #include "XRCreativeTeleporter.h"
 #include "MotionControllerComponent.h"
-#include "XRDeviceVisualizationComponent.h"
 #include "Components/StaticMeshComponent.h"
 
 
diff --git a/Engine/Plugins/Experimental/XRCreativeFramework/Source/XRCreative/Public/XRCreativeTeleporter.h b/Engine/Plugins/Experimental/XRCreativeFramework/Source/XRCreative/Public/XRCreativeTeleporter.h
--- a/Engine/Plugins/Experimental/XRCreativeFramework/Source/XRCreative/Public/XRCreativeTeleporter.h
+++ b/Engine/Plugins/Experimental/XRCreativeFramework/Source/XRCreative/Public/XRCreativeTeleporter.h
@@ -5,7 +5,6 @@
 #include "GameFramework/Actor.h"
 #include "XRCreativeTeleporter.generated.h"
 
-class UXRDeviceVisualizationComponent;
 class UMotionControllerComponent;
 class UStaticMeshComponent;
 
@@ -52,9 +51,5 @@ protected:
 
 	UPROPERTY(EditAnywhere, BlueprintReadOnly, Category="XR Creative Teleporter")
 	TObjectPtr<UStaticMeshComponent> HMDVisual;
-
-	UPROPERTY(EditAnywhere, BlueprintReadOnly, Category="XR Creative Teleporter")
-	TObjectPtr<UXRDeviceVisualizationComponent> TestVisual;
-
-	
+		
 };
