XRCreative - renamed mapping contexts from Default & Left to RightHanded & LeftHanded
Virtual Scouting - Fixed up all affected content assets

#virtualized

diff --git a/Engine/Plugins/Experimental/XRCreativeFramework/Source/XRCreative/Private/XRCreativeToolset.cpp b/Engine/Plugins/Experimental/XRCreativeFramework/Source/XRCreative/Private/XRCreativeToolset.cpp
--- a/Engine/Plugins/Experimental/XRCreativeFramework/Source/XRCreative/Private/XRCreativeToolset.cpp
+++ b/Engine/Plugins/Experimental/XRCreativeFramework/Source/XRCreative/Private/XRCreativeToolset.cpp
@@ -14,19 +14,19 @@ DEFINE_LOG_CATEGORY(LogXRCreativeToolset);
 
 UInputMappingContext* UXRCreativeBlueprintableTool::GetToolInputMappingContext()
 {
-	UInputMappingContext* InputMappingContext = DefaultToolInputMappingContext;
+	UInputMappingContext* InputMappingContext = RightHandedInputMappingContext;
 	
 	#if WITH_EDITOR
 		UXRCreativeEditorSettings* Settings = UXRCreativeEditorSettings::GetXRCreativeEditorSettings();
 
-		if (Settings->Handedness == EXRCreativeHandedness::Left && LeftToolInputMappingContext)
+		if (Settings->Handedness == EXRCreativeHandedness::Left && LeftHandedInputMappingContext)
 		{
-			InputMappingContext = LeftToolInputMappingContext;
+			InputMappingContext = LeftHandedInputMappingContext;
 		}
-		else if (Settings->Handedness == EXRCreativeHandedness::Left && !LeftToolInputMappingContext)
+		else if (Settings->Handedness == EXRCreativeHandedness::Left && !LeftHandedInputMappingContext)
 		{
-			InputMappingContext = DefaultToolInputMappingContext;
-			UE_LOG(LogXRCreativeToolset, Warning, TEXT("Handedness is Left but no Left Input Mapping Context found in Toolset - Using Default."));
+			InputMappingContext = RightHandedInputMappingContext;
+			UE_LOG(LogXRCreativeToolset, Warning, TEXT("Handedness is Left but no Left Handed Input Mapping Context found in Toolset - Using RightHand."));
 		}
 	
 	#endif
diff --git a/Engine/Plugins/Experimental/XRCreativeFramework/Source/XRCreative/Public/XRCreativeToolset.h b/Engine/Plugins/Experimental/XRCreativeFramework/Source/XRCreative/Public/XRCreativeToolset.h
--- a/Engine/Plugins/Experimental/XRCreativeFramework/Source/XRCreative/Public/XRCreativeToolset.h
+++ b/Engine/Plugins/Experimental/XRCreativeFramework/Source/XRCreative/Public/XRCreativeToolset.h
@@ -92,10 +92,10 @@ protected:
 	TSubclassOf<AXRCreativeToolActor> ToolActor;
 
 	UPROPERTY(BlueprintReadWrite, EditAnywhere, Category="XR Creative")
-	TObjectPtr<UInputMappingContext> DefaultToolInputMappingContext;
+	TObjectPtr<UInputMappingContext> RightHandedInputMappingContext;
 
 	UPROPERTY(BlueprintReadWrite, EditAnywhere, Category="XR Creative")
-	TObjectPtr<UInputMappingContext> LeftToolInputMappingContext;
+	TObjectPtr<UInputMappingContext> LeftHandedInputMappingContext;
 
 	
 };
@@ -145,11 +145,11 @@ public:
 
 	/** Default Input Mapping is used for Right-Handed users, or if no LeftInputMappingContext is provided. */
 	UPROPERTY(EditDefaultsOnly, BlueprintReadOnly, Category="XR Creative")
-	TObjectPtr<UInputMappingContext> DefaultInputMappingContext;
+	TObjectPtr<UInputMappingContext> RightHandedInputMappingContext;
 	
 	/** If Handedness is selected in XRCreative Settings, uses this entry in place of Default/Right  */
 	UPROPERTY(EditDefaultsOnly, BlueprintReadOnly, Category="XR Creative")
-	TObjectPtr<UInputMappingContext> LeftInputMappingContext;
+	TObjectPtr<UInputMappingContext> LeftHandedInputMappingContext;
 	
 	UPROPERTY(EditDefaultsOnly, BlueprintReadOnly, Category="XR Creative")
 	TArray<FXRCreativeToolEntry> Tools;
diff --git a/Engine/Plugins/Experimental/XRCreativeFramework/Source/XRCreativeEditor/Private/XRCreativeVREditorMode.cpp b/Engine/Plugins/Experimental/XRCreativeFramework/Source/XRCreativeEditor/Private/XRCreativeVREditorMode.cpp
--- a/Engine/Plugins/Experimental/XRCreativeFramework/Source/XRCreativeEditor/Private/XRCreativeVREditorMode.cpp
+++ b/Engine/Plugins/Experimental/XRCreativeFramework/Source/XRCreativeEditor/Private/XRCreativeVREditorMode.cpp
@@ -233,14 +233,14 @@ void UXRCreativeVREditorMode::EnableStereo()
 
 	if (const UXRCreativeToolset* Toolset = Avatar->GetToolset())
 	{
-		if (Toolset->DefaultInputMappingContext)
+		if (Toolset->RightHandedInputMappingContext)
 		{
-			Contexts.Add(Toolset->DefaultInputMappingContext);
+			Contexts.Add(Toolset->RightHandedInputMappingContext);
 		}
 
-		if (Toolset->LeftInputMappingContext)
+		if (Toolset->LeftHandedInputMappingContext)
 		{
-			Contexts.Add(Toolset->LeftInputMappingContext);
+			Contexts.Add(Toolset->LeftHandedInputMappingContext);
 		}
 	}
 	
