XR Creative Framework: Unify actor/component ignore for PointerComponent and SelectionInteraction.

Also removes the unused "subsystem helpers" stub.

diff --git a/Engine/Plugins/Experimental/XRCreativeFramework/Source/XRCreative/Private/ITF/SelectionInteraction.cpp b/Engine/Plugins/Experimental/XRCreativeFramework/Source/XRCreative/Private/ITF/SelectionInteraction.cpp
--- a/Engine/Plugins/Experimental/XRCreativeFramework/Source/XRCreative/Private/ITF/SelectionInteraction.cpp
+++ b/Engine/Plugins/Experimental/XRCreativeFramework/Source/XRCreative/Private/ITF/SelectionInteraction.cpp
@@ -13,12 +13,23 @@
 #endif
 
 
-void UXRCreativeSelectionInteraction::Initialize(UTypedElementSelectionSet* InSelectionSet, FActorPredicate InCanSelectCallback)
+void UXRCreativeSelectionInteraction::Initialize(
+	UTypedElementSelectionSet* InSelectionSet,
+	FActorPredicate InCanSelectCallback,
+	FTraceMethod InTraceCallback)
 {
 	ensure(InSelectionSet);
-
 	WeakSelectionSet = InSelectionSet;
-	CanSelectCallback = MoveTemp(InCanSelectCallback);
+
+	if (InCanSelectCallback)
+	{
+		CanSelectCallback = MoveTemp(InCanSelectCallback);
+	}
+
+	if (InTraceCallback)
+	{
+		TraceCallback = MoveTemp(InTraceCallback);
+	}
 
 	// create click behavior and set ourselves as click target
 	ClickBehavior = NewObject<USingleClickInputBehavior>();
@@ -61,15 +72,12 @@ FInputRayHit UXRCreativeSelectionInteraction::IsHitByClick(const FInputDeviceRay
 {
 	FInputRayHit RayHit;
 
-	FHitResult Result;
-	FCollisionObjectQueryParams QueryParams(FCollisionObjectQueryParams::AllObjects);
-	const bool bBlockingHit = GetWorld()->LineTraceSingleByObjectType(Result,
-		ClickPos.WorldRay.Origin, ClickPos.WorldRay.PointAt(999999), QueryParams);
-	AActor* HitActor = Result.GetActor();
-	if (bBlockingHit && CanSelectCallback(HitActor))
+	const FHitResult HitResult = TraceCallback(ClickPos);
+	AActor* HitActor = HitResult.GetActor();
+	if (HitResult.bBlockingHit && CanSelectCallback(HitActor))
 	{
 		RayHit.bHit = true;
-		RayHit.HitDepth = Result.Distance;
+		RayHit.HitDepth = HitResult.Distance;
 		RayHit.bHasHitNormal = true;
 		RayHit.SetHitObject(HitActor);
 	}
@@ -100,6 +108,15 @@ FTypedElementHandle UXRCreativeSelectionInteraction::AcquireActorElementHandle(c
 }
 
 
+FHitResult UXRCreativeSelectionInteraction::DefaultTrace(const FInputDeviceRay& InRay) const
+{
+	FCollisionObjectQueryParams QueryParams(FCollisionObjectQueryParams::AllObjects);
+	FHitResult Result;
+	GetWorld()->LineTraceSingleByObjectType(Result, InRay.WorldRay.Origin, InRay.WorldRay.PointAt(RayLength), QueryParams);
+	return Result;
+}
+
+
 void UXRCreativeSelectionInteraction::OnClicked(const FInputDeviceRay& ClickPos)
 {
 	UTypedElementSelectionSet* SelectionSet = WeakSelectionSet.Get();
@@ -109,10 +126,8 @@ void UXRCreativeSelectionInteraction::OnClicked(const FInputDeviceRay& ClickPos)
 	}
 
 	FCollisionObjectQueryParams QueryParams(FCollisionObjectQueryParams::AllObjects);
-	FHitResult Result;
-	const bool bBlockingHit = GetWorld()->LineTraceSingleByObjectType(Result,
-		ClickPos.WorldRay.Origin, ClickPos.WorldRay.PointAt(999999), QueryParams);
-	if (ensure(bBlockingHit))
+	const FHitResult Result = TraceCallback(ClickPos);
+	if (ensure(Result.bBlockingHit))
 	{
 		ensure(CanSelectCallback(Result.GetActor()));
 
diff --git a/Engine/Plugins/Experimental/XRCreativeFramework/Source/XRCreative/Private/ITF/SelectionInteraction.h b/Engine/Plugins/Experimental/XRCreativeFramework/Source/XRCreative/Private/ITF/SelectionInteraction.h
--- a/Engine/Plugins/Experimental/XRCreativeFramework/Source/XRCreative/Private/ITF/SelectionInteraction.h
+++ b/Engine/Plugins/Experimental/XRCreativeFramework/Source/XRCreative/Private/ITF/SelectionInteraction.h
@@ -6,6 +6,7 @@
 #include "BaseBehaviors/SingleClickBehavior.h"
 #include "Elements/Actor/ActorElementData.h"
 #include "Elements/Framework/TypedElementOwnerStore.h"
+#include "Engine/HitResult.h"
 #include "InputBehaviorSet.h"
 #include "SelectionInteraction.generated.h"
 
@@ -19,15 +20,23 @@ class UXRCreativeSelectionInteraction : public UObject, public IInputBehaviorSou
 	GENERATED_BODY()
 
 public:
+	static constexpr double RayLength = 999999.0;
+
 	using FActorPredicate = TUniqueFunction<bool(AActor*)>;
+	using FTraceMethod = TUniqueFunction<FHitResult(const FInputDeviceRay&)>;
 
 	/**
 	 * Set up the Interaction, creates and registers Behaviors/etc. 
 	 * 
 	 * @param InSelectionSet the typed element selection set we maintain a weak pointer to and operate on
 	 * @param InCanSelectCallback this function will be called to determine if the selection can be changed to the specified actor (or null); when a tool is active, we may wish to lock selection
+	 * @param InTraceCallback can be provided to override the default scene trace implementation
 	 */
-	void Initialize(UTypedElementSelectionSet* InSelectionSet, FActorPredicate InCanSelectCallback);
+	void Initialize(
+		UTypedElementSelectionSet* InSelectionSet,
+		FActorPredicate InCanSelectCallback = nullptr,
+		FTraceMethod InTraceCallback = nullptr);
+
 	void Shutdown();
 
 	// IInputBehaviorSource interface
@@ -43,6 +52,8 @@ public:
 protected:
 	FTypedElementHandle AcquireActorElementHandle(const AActor* Actor, const bool bAllowCreate);
 
+	FHitResult DefaultTrace(const FInputDeviceRay& InRay) const;
+
 protected:
 	// click-to-select behavior
 	UPROPERTY()
@@ -61,6 +72,8 @@ protected:
 	// default predicate allows anything
 	FActorPredicate CanSelectCallback = [](AActor*) { return true; };
 
+	FTraceMethod TraceCallback = [this](const FInputDeviceRay& InRay) { return DefaultTrace(InRay); };
+
 	// flags used to identify behavior modifier keys/buttons
 	static const int AddToSelectionModifier = 1;
 	bool bAddToSelectionEnabled = false;
diff --git a/Engine/Plugins/Experimental/XRCreativeFramework/Source/XRCreative/Private/ITF/XRCreativeITFComponent.cpp b/Engine/Plugins/Experimental/XRCreativeFramework/Source/XRCreative/Private/ITF/XRCreativeITFComponent.cpp
--- a/Engine/Plugins/Experimental/XRCreativeFramework/Source/XRCreative/Private/ITF/XRCreativeITFComponent.cpp
+++ b/Engine/Plugins/Experimental/XRCreativeFramework/Source/XRCreative/Private/ITF/XRCreativeITFComponent.cpp
@@ -414,35 +414,51 @@ void UXRCreativeITFComponent::InitializeComponent()
 #endif
 
 	// register selection interaction
-	SelectionInteraction = NewObject<UXRCreativeSelectionInteraction>(this);
-	SelectionInteraction->Initialize(GetSelectionSet(),
-		[this](AActor* SelectionCandidate)
-		{
-			if (HaveActiveTool())
+	{
+		auto InteractionActorCallback = [this](AActor* SelectionCandidate)
 			{
-				return false;
-			}
+				if (HaveActiveTool())
+				{
+					return false;
+				}
 
-			if (CanSelectPredicate.IsBound())
-			{
-				return CanSelectPredicate.Execute(SelectionCandidate);
-			}
+				if (CanSelectPredicate.IsBound())
+				{
+					return CanSelectPredicate.Execute(SelectionCandidate);
+				}
 
-			if (SelectionCandidate)
-			{
-				for (const TSubclassOf<AActor>& DisallowedClass : UnselectableActorClasses)
+				if (SelectionCandidate)
 				{
-					if (DisallowedClass.Get() && SelectionCandidate->IsA(DisallowedClass))
+					for (const TSubclassOf<AActor>& DisallowedClass : UnselectableActorClasses)
 					{
-						return false;
+						if (DisallowedClass.Get() && SelectionCandidate->IsA(DisallowedClass))
+						{
+							return false;
+						}
 					}
 				}
-			}
 
-			return true;
-		}
-	);
-	ToolsContext->InputRouter->RegisterSource(SelectionInteraction);
+				return true;
+			};
+
+		auto InteractionTraceCallback = [this](const FInputDeviceRay& InRay) -> FHitResult
+			{
+				FHitResult Result;
+				if (ensure(PointerComponent))
+				{
+					GetWorld()->LineTraceSingleByChannel(Result, InRay.WorldRay.Origin,
+						InRay.WorldRay.PointAt(UXRCreativeSelectionInteraction::RayLength),
+						ECC_Visibility, PointerComponent->GetQueryParams());
+				}
+				return Result;
+			};
+
+		SelectionInteraction = NewObject<UXRCreativeSelectionInteraction>(this);
+		SelectionInteraction->Initialize(GetSelectionSet(),
+			MoveTemp(InteractionActorCallback),
+			MoveTemp(InteractionTraceCallback));
+		ToolsContext->InputRouter->RegisterSource(SelectionInteraction);
+	}
 
 	// create transform interaction
 	TransformInteraction = NewObject<UXRCreativeTransformInteraction>(this);
diff --git a/Engine/Plugins/Experimental/XRCreativeFramework/Source/XRCreative/Private/XRCreativePointerComponent.cpp b/Engine/Plugins/Experimental/XRCreativeFramework/Source/XRCreative/Private/XRCreativePointerComponent.cpp
--- a/Engine/Plugins/Experimental/XRCreativeFramework/Source/XRCreative/Private/XRCreativePointerComponent.cpp
+++ b/Engine/Plugins/Experimental/XRCreativeFramework/Source/XRCreative/Private/XRCreativePointerComponent.cpp
@@ -47,6 +47,16 @@ void UXRCreativePointerComponent::UninitializeComponent()
 }
 
 
+FCollisionQueryParams UXRCreativePointerComponent::GetQueryParams() const
+{
+	FCollisionQueryParams QueryParams;
+	QueryParams.AddIgnoredActors(IgnoredActors);
+	QueryParams.AddIgnoredComponents(IgnoredComponents);
+	QueryParams.bTraceComplex = true;
+	return QueryParams;
+}
+
+
 void UXRCreativePointerComponent::TickComponent(float InDeltaTime, enum ELevelTick InTickType, FActorComponentTickFunction* InThisTickFunction)
 {
 	Super::TickComponent(InDeltaTime, InTickType, InThisTickFunction);
@@ -54,12 +64,7 @@ void UXRCreativePointerComponent::TickComponent(float InDeltaTime, enum ELevelTi
 	RawTraceEnd = GetComponentLocation() + (GetForwardVector() * TraceMaxLength);
 	const FVector FilteredTraceEnd = SmoothingFilter->Filter(RawTraceEnd, InDeltaTime);
 
-	FCollisionQueryParams QueryParams;
-	QueryParams.AddIgnoredActors(IgnoredActors);
-	QueryParams.AddIgnoredComponents(IgnoredComponents);
-	QueryParams.bTraceComplex = true;
-
-	GetWorld()->LineTraceSingleByChannel(HitResult, GetComponentLocation(), FilteredTraceEnd, ECC_Visibility, QueryParams);
+	GetWorld()->LineTraceSingleByChannel(HitResult, GetComponentLocation(), FilteredTraceEnd, ECC_Visibility, GetQueryParams());
 }
 
 
diff --git a/Engine/Plugins/Experimental/XRCreativeFramework/Source/XRCreative/Private/XRCreativeSubsystem.cpp b/Engine/Plugins/Experimental/XRCreativeFramework/Source/XRCreative/Private/XRCreativeSubsystem.cpp
--- a/Engine/Plugins/Experimental/XRCreativeFramework/Source/XRCreative/Private/XRCreativeSubsystem.cpp
+++ b/Engine/Plugins/Experimental/XRCreativeFramework/Source/XRCreative/Private/XRCreativeSubsystem.cpp
@@ -2,7 +2,6 @@
 
 #include "XRCreativeSubsystem.h"
 #include "XRCreativeSettings.h"
-#include "Misc/CoreDelegates.h"
 #include "Types/MVVMViewModelCollection.h"
 #include "UObject/Package.h"
 
@@ -15,8 +14,6 @@ void UXRCreativeSubsystem::Initialize(FSubsystemCollectionBase& Collection)
 {
 	Super::Initialize(Collection);
 	ViewModelCollection = NewObject<UMVVMViewModelCollectionObject>(this);
-
-	EngineInitCompleteDelegate = FCoreDelegates::OnFEngineLoopInitComplete.AddUObject(this, &UXRCreativeSubsystem::OnEngineInitComplete);
 }
 
 
@@ -43,16 +40,3 @@ void UXRCreativeSubsystem::ExitVRMode()
 	}
 }
 #endif // #if WITH_EDITOR
-
-
-void UXRCreativeSubsystem::OnEngineInitComplete()
-{
-	EngineInitCompleteDelegate.Reset();
-
-	const UXRCreativeSettings* Settings = UXRCreativeSettings::GetXRCreativeSettings();
-	if (UClass* HelperClass = Settings->SubsystemHelperClass.LoadSynchronous())
-	{
-		Helpers = NewObject<UXRCreativeSubsystemHelper>(GetTransientPackage(), HelperClass);
-	}
-}
-
diff --git a/Engine/Plugins/Experimental/XRCreativeFramework/Source/XRCreative/Public/XRCreativePointerComponent.h b/Engine/Plugins/Experimental/XRCreativeFramework/Source/XRCreative/Public/XRCreativePointerComponent.h
--- a/Engine/Plugins/Experimental/XRCreativeFramework/Source/XRCreative/Public/XRCreativePointerComponent.h
+++ b/Engine/Plugins/Experimental/XRCreativeFramework/Source/XRCreative/Public/XRCreativePointerComponent.h
@@ -2,6 +2,7 @@
 
 #pragma once
 
+#include "CollisionQueryParams.h"
 #include "Components/SceneComponent.h"
 #include "Engine/HitResult.h"
 #include "XRCreativePointerComponent.generated.h"
@@ -48,6 +49,8 @@ public:
 	UFUNCTION(BlueprintCallable, Category="XR Creative")
 	void SetEnabled(bool bInEnabled) { bEnabled = bInEnabled; }
 
+	FCollisionQueryParams GetQueryParams() const;
+
 protected:
 	UPROPERTY(EditAnywhere, BlueprintReadWrite, Category="XR Creative")
 	float TraceMaxLength;
diff --git a/Engine/Plugins/Experimental/XRCreativeFramework/Source/XRCreative/Public/XRCreativeSettings.h b/Engine/Plugins/Experimental/XRCreativeFramework/Source/XRCreative/Public/XRCreativeSettings.h
--- a/Engine/Plugins/Experimental/XRCreativeFramework/Source/XRCreative/Public/XRCreativeSettings.h
+++ b/Engine/Plugins/Experimental/XRCreativeFramework/Source/XRCreative/Public/XRCreativeSettings.h
@@ -4,6 +4,7 @@
 
 #include "XRCreativeSettings.generated.h"
 
+
 UENUM(BlueprintType)
 enum class EXRCreativeHandedness : uint8
 {
@@ -11,24 +12,21 @@ enum class EXRCreativeHandedness : uint8
 	Right	UMETA(DisplayName = "Right"),
 };
 
+
 /**
  * Per project settings for XRCreative.
  */
 UCLASS(Config=XRCreativeSettings, DefaultConfig, DisplayName="XR Creative")
-
 class XRCREATIVE_API UXRCreativeSettings : public UObject
 {
 	GENERATED_BODY()
 	
 public:
-
-	UPROPERTY(Config, EditAnywhere, BlueprintReadWrite, Category="XR Creative")
-	TSoftClassPtr<class UXRCreativeSubsystemHelper> SubsystemHelperClass;
-
 	UFUNCTION(BlueprintPure, Category="XR Creative")
 	static UXRCreativeSettings* GetXRCreativeSettings();
 };
 
+
 /**
  * Per user settings for XRCreative Editor.
  */
diff --git a/Engine/Plugins/Experimental/XRCreativeFramework/Source/XRCreative/Public/XRCreativeSubsystem.h b/Engine/Plugins/Experimental/XRCreativeFramework/Source/XRCreative/Public/XRCreativeSubsystem.h
--- a/Engine/Plugins/Experimental/XRCreativeFramework/Source/XRCreative/Public/XRCreativeSubsystem.h
+++ b/Engine/Plugins/Experimental/XRCreativeFramework/Source/XRCreative/Public/XRCreativeSubsystem.h
@@ -9,13 +9,6 @@
 class UMVVMViewModelCollectionObject;
 
 
-UCLASS(Abstract, Blueprintable)
-class XRCREATIVE_API UXRCreativeSubsystemHelper : public UObject
-{
-	GENERATED_BODY()
-};
-
-
 UCLASS()
 class XRCREATIVE_API UXRCreativeSubsystem : public UEngineSubsystem
 {
@@ -24,7 +17,7 @@ class XRCREATIVE_API UXRCreativeSubsystem : public UEngineSubsystem
 public:
 	virtual void Initialize(FSubsystemCollectionBase& Collection) override;
 
-	UFUNCTION(BlueprintCallable, Category = "XR Creative|Viewmodel")
+	UFUNCTION(BlueprintCallable, Category="XR Creative|Viewmodel")
 	UMVVMViewModelCollectionObject* GetViewModelCollection() const
 	{
 		return ViewModelCollection;
@@ -43,15 +36,4 @@ public:
 private:
 	UPROPERTY(Transient)
 	TObjectPtr<UMVVMViewModelCollectionObject> ViewModelCollection;
-
-protected:
-	void OnEngineInitComplete();
-
-protected:
-	UPROPERTY(BlueprintReadOnly, Category="XR Creative")
-	TObjectPtr<UXRCreativeSubsystemHelper> Helpers;
-
-	FDelegateHandle EngineInitCompleteDelegate;
-
-
 };
