XRCreative - Added early out on InteractionActorCallback if the pointercomponent is not enabled.

diff --git a/Engine/Plugins/Experimental/XRCreativeFramework/Source/XRCreative/Private/ITF/XRCreativeITFComponent.cpp b/Engine/Plugins/Experimental/XRCreativeFramework/Source/XRCreative/Private/ITF/XRCreativeITFComponent.cpp
--- a/Engine/Plugins/Experimental/XRCreativeFramework/Source/XRCreative/Private/ITF/XRCreativeITFComponent.cpp
+++ b/Engine/Plugins/Experimental/XRCreativeFramework/Source/XRCreative/Private/ITF/XRCreativeITFComponent.cpp
@@ -417,7 +417,7 @@ void UXRCreativeITFComponent::InitializeComponent()
 	{
 		auto InteractionActorCallback = [this](AActor* SelectionCandidate)
 			{
-				if (HaveActiveTool())
+				if (HaveActiveTool() || !PointerComponent->IsEnabled())
 				{
 					return false;
 				}
