[Virtual Scouting] Fix for Left-handed grab selection not working.

Moved call to SetPointerComponent function later in the Avatar's lifecycle - from object init to construction so the setting for handedness is called when entering VR (spawning the Avatar), not just when the editor first starts.

diff --git a/Engine/Plugins/Experimental/XRCreativeFramework/Source/XRCreative/Private/XRCreativeAvatar.cpp b/Engine/Plugins/Experimental/XRCreativeFramework/Source/XRCreative/Private/XRCreativeAvatar.cpp
--- a/Engine/Plugins/Experimental/XRCreativeFramework/Source/XRCreative/Private/XRCreativeAvatar.cpp
+++ b/Engine/Plugins/Experimental/XRCreativeFramework/Source/XRCreative/Private/XRCreativeAvatar.cpp
@@ -100,18 +100,6 @@ AXRCreativeAvatar::AXRCreativeAvatar(const FObjectInitializer& ObjectInitializer
 	ToolsComponent->SetPointerComponent(RightControllerPointer);
 
 	BaseEyeHeight = 0.0f;
-
-#if WITH_EDITOR
-	UXRCreativeEditorSettings* Settings = UXRCreativeEditorSettings::GetXRCreativeEditorSettings();
-	UE_LOG(LogXRCreative, Log, TEXT("Handedness: %s"), *UEnum::GetValueAsString(Settings->Handedness) );
-
-	if (Settings->Handedness == EXRCreativeHandedness::Left)
-	{
-		MenuWidget->SetupAttachment(RightControllerAim);
-		WidgetInteraction->SetupAttachment(LeftControllerAim);
-		ToolsComponent->SetPointerComponent(LeftControllerPointer);
-	}
-#endif
 }
 
 
@@ -142,8 +130,14 @@ void AXRCreativeAvatar::OnConstruction(const FTransform& InTransform)
 		// We're in editor, and the base class won't initialize our components, so do it manually.
 		ToolsComponent->InitializeComponent();
 	}
+	UXRCreativeEditorSettings* Settings = UXRCreativeEditorSettings::GetXRCreativeEditorSettings();
+	UE_LOG(LogXRCreative, Log, TEXT("Handedness: %s"), *UEnum::GetValueAsString(Settings->Handedness) );
+	
+	if (Settings->Handedness == EXRCreativeHandedness::Left)
+	{
+		ToolsComponent->SetPointerComponent(LeftControllerPointer);
+	}
 #endif
-
 	if (!IsTemplate())
 	{
 #if WITH_EDITOR
@@ -152,7 +146,6 @@ void AXRCreativeAvatar::OnConstruction(const FTransform& InTransform)
 	}
 }
 
-
 void AXRCreativeAvatar::BeginDestroy()
 {
 	UnregisterInputComponent();
