XR Creative Framework: Experimental blueprintable gizmo implementation.

#virtualized

diff --git a/Engine/Plugins/Experimental/XRCreativeFramework/Source/XRCreative/Private/ITF/TransformInteraction.cpp b/Engine/Plugins/Experimental/XRCreativeFramework/Source/XRCreative/Private/ITF/TransformInteraction.cpp
--- a/Engine/Plugins/Experimental/XRCreativeFramework/Source/XRCreative/Private/ITF/TransformInteraction.cpp
+++ b/Engine/Plugins/Experimental/XRCreativeFramework/Source/XRCreative/Private/ITF/TransformInteraction.cpp
@@ -13,56 +13,11 @@
 #include "InteractiveGizmoManager.h"
 
 
-class FXRCreativeTransformGizmoActorFactory : public FCombinedTransformGizmoActorFactory
-{
-public:
-	FXRCreativeTransformGizmoActorFactory(UGizmoViewContext* InGizmoViewContext)
-		: FCombinedTransformGizmoActorFactory(InGizmoViewContext)
-	{
-		ensure(InGizmoViewContext);
-	}
-
-	virtual ACombinedTransformGizmoActor* CreateNewGizmoActor(UWorld* World) const override
-	{
-		ACombinedTransformGizmoActor* GizmoActor = nullptr;
-
-		switch (EnableElements)
-		{
-		case ETransformGizmoSubElements::FullTranslateRotateScale:
-			GizmoActor = World->SpawnActor<AXRCreativeTRSGizmoActor>(FVector::ZeroVector, FRotator::ZeroRotator, FActorSpawnParameters());
-			break;
-		case ETransformGizmoSubElements::TranslateRotateUniformScale:
-			GizmoActor = World->SpawnActor<AXRCreativeTRUSGizmoActor>(FVector::ZeroVector, FRotator::ZeroRotator, FActorSpawnParameters());
-			break;
-		case ETransformGizmoSubElements::StandardTranslateRotate:
-			GizmoActor = World->SpawnActor<AXRCreativeTRGizmoActor>(FVector::ZeroVector, FRotator::ZeroRotator, FActorSpawnParameters());
-			break;
-		default:
-			ensure(false); // Unexpected case, not handled above, fall back to parent default behavior
-			return FCombinedTransformGizmoActorFactory::CreateNewGizmoActor(World);
-		}
-
-		TInlineComponentArray<UGizmoBaseComponent*> GizmoComponents;
-		GizmoActor->GetComponents(GizmoComponents);
-		for (UGizmoBaseComponent* GizmoComp : GizmoComponents)
-		{
-			GizmoComp->PixelHitDistanceThreshold = 14.0f;
-			GizmoComp->SetGizmoViewContext(GizmoViewContext);
-			GizmoComp->NotifyExternalPropertyUpdates();
-		}
-
-		return GizmoActor;
-	}
-};
-
-
-//////////////////////////////////////////////////////////////////////////
-
-
 const FString UXRCreativeTransformInteraction::GizmoBuilderIdentifier("XRCreativeGizmo");
 
 
 void UXRCreativeTransformInteraction::Initialize(
+	TSharedRef<FCombinedTransformGizmoActorFactory> InGizmoActorFactory,
 	UTypedElementSelectionSet* InSelectionSet,
 	UInteractiveGizmoManager* InGizmoManager,
 	TUniqueFunction<bool()> InGizmoEnabledCallback)
@@ -70,13 +25,11 @@ void UXRCreativeTransformInteraction::Initialize(
 	check(InSelectionSet && IsValid(InSelectionSet));
 	check(InGizmoManager && IsValid(InGizmoManager));
 
+	GizmoActorFactory = InGizmoActorFactory;
 	WeakSelectionSet = InSelectionSet;
 	WeakGizmoManager = InGizmoManager;
 	GizmoEnabledCallback = MoveTemp(InGizmoEnabledCallback);
 
-	UGizmoViewContext* GizmoViewContext = InGizmoManager->GetContextObjectStore()->FindContext<UGizmoViewContext>();
-	GizmoActorFactory = MakeShared<FXRCreativeTransformGizmoActorFactory>(GizmoViewContext);
-
 	UXRCreativeGizmoBuilder* GizmoBuilder = NewObject<UXRCreativeGizmoBuilder>();
 	GizmoBuilder->GizmoActorBuilder = GizmoActorFactory;
 	InGizmoManager->RegisterGizmoType(GizmoBuilderIdentifier, GizmoBuilder);
@@ -191,6 +144,10 @@ void UXRCreativeTransformInteraction::UpdateGizmoTargets(const UTypedElementSele
 	TransformGizmo = CastChecked<UCombinedTransformGizmo>(GizmoManager->CreateGizmo(GizmoBuilderIdentifier, FString(), this));
 	TransformGizmo->SetActiveTarget(TransformProxy);
 
+	AXRCreativeCombinedTransformGizmoActor* NewGizmoActor =
+		CastChecked<AXRCreativeCombinedTransformGizmoActor>(TransformGizmo->GetGizmoActor());
+	NewGizmoActor->WeakGizmoManager = GizmoManager;
+
 	// optionally ignore coordinate system setting
 	//TransformGizmo->bUseContextCoordinateSystem = false;
 	//TransformGizmo->CurrentCoordinateSystem = EToolContextCoordinateSystem::Local;
diff --git a/Engine/Plugins/Experimental/XRCreativeFramework/Source/XRCreative/Private/ITF/TransformInteraction.h b/Engine/Plugins/Experimental/XRCreativeFramework/Source/XRCreative/Private/ITF/TransformInteraction.h
--- a/Engine/Plugins/Experimental/XRCreativeFramework/Source/XRCreative/Private/ITF/TransformInteraction.h
+++ b/Engine/Plugins/Experimental/XRCreativeFramework/Source/XRCreative/Private/ITF/TransformInteraction.h
@@ -6,6 +6,7 @@
 #include "TransformInteraction.generated.h"
 
 
+class FCombinedTransformGizmoActorFactory;
 class FXRCreativeTransformGizmoActorFactory;
 class UTransformProxy;
 class UCombinedTransformGizmo;
@@ -34,7 +35,8 @@ public:
 	 * Set up the transform interaction. 
 	 * @param InGizmoEnabledCallback callback that determines if Gizmo should be created and visible. For example during a Tool we generally want to hide the TRS Gizmo.
 	 */
-	void Initialize(UTypedElementSelectionSet* InSelectionSet,
+	void Initialize(TSharedRef<FCombinedTransformGizmoActorFactory> InGizmoActorFactory,
+		UTypedElementSelectionSet* InSelectionSet,
 		UInteractiveGizmoManager* InGizmoManager,
 		TUniqueFunction<bool()> InGizmoEnabledCallback);
 
@@ -62,7 +64,7 @@ protected:
 	TWeakObjectPtr<UTypedElementSelectionSet> WeakSelectionSet;
 	TWeakObjectPtr<UInteractiveGizmoManager> WeakGizmoManager;
 
-	TSharedPtr<FXRCreativeTransformGizmoActorFactory> GizmoActorFactory;
+	TSharedPtr<FCombinedTransformGizmoActorFactory> GizmoActorFactory;
 
 	void UpdateGizmoTargets(const UTypedElementSelectionSet* InSelectionSet);
 
diff --git a/Engine/Plugins/Experimental/XRCreativeFramework/Source/XRCreative/Private/ITF/XRCreativeGizmos.cpp b/Engine/Plugins/Experimental/XRCreativeFramework/Source/XRCreative/Private/ITF/XRCreativeGizmos.cpp
--- a/Engine/Plugins/Experimental/XRCreativeFramework/Source/XRCreative/Private/ITF/XRCreativeGizmos.cpp
+++ b/Engine/Plugins/Experimental/XRCreativeFramework/Source/XRCreative/Private/ITF/XRCreativeGizmos.cpp
@@ -1,11 +1,17 @@
 // Copyright Epic Games, Inc. All Rights Reserved.
 
 #include "XRCreativeGizmos.h"
+#include "BaseGizmos/AxisAngleGizmo.h"
+#include "BaseGizmos/AxisPositionGizmo.h"
 #include "BaseGizmos/GizmoArrowComponent.h"
 #include "BaseGizmos/GizmoBoxComponent.h"
 #include "BaseGizmos/GizmoCircleComponent.h"
 #include "BaseGizmos/GizmoRectangleComponent.h"
 #include "BaseGizmos/GizmoViewContext.h"
+#include "BaseGizmos/PlanePositionGizmo.h"
+#include "BaseGizmos/TransformSubGizmoUtil.h"
+#include "ContextObjectStore.h"
+#include "Materials/MaterialInstanceDynamic.h"
 
 
 UXRCreativeGizmoBuilder::UXRCreativeGizmoBuilder()
@@ -15,187 +21,585 @@ UXRCreativeGizmoBuilder::UXRCreativeGizmoBuilder()
 	AxisAngleBuilderIdentifier = UInteractiveGizmoManager::DefaultAxisAngleBuilderIdentifier;
 }
 
+UInteractiveGizmo* UXRCreativeGizmoBuilder::BuildGizmo(const FToolBuilderState& SceneState) const
+{
+	// FIXME: Have to duplicate base implementation to customize instantiated gizmo type
+	// Substituting UXRCreativeGizmo on the next line is the only difference from the base
+	UCombinedTransformGizmo* NewGizmo = NewObject<UXRCreativeGizmo>(SceneState.GizmoManager);
+	NewGizmo->SetWorld(SceneState.World);
+
+	UGizmoViewContext* GizmoViewContext = SceneState.ToolManager->GetContextObjectStore()->FindContext<UGizmoViewContext>();
+	check(IsValid(GizmoViewContext));
+
+	// use default gizmo actor if client has not given us a new builder
+	NewGizmo->SetGizmoActorBuilder(GizmoActorBuilder ? GizmoActorBuilder : MakeShared<FCombinedTransformGizmoActorFactory>(GizmoViewContext));
+
+	NewGizmo->SetSubGizmoBuilderIdentifiers(AxisPositionBuilderIdentifier, PlanePositionBuilderIdentifier, AxisAngleBuilderIdentifier);
+
+	// override default hover function if proposed
+	if (UpdateHoverFunction)
+	{
+		NewGizmo->SetUpdateHoverFunction(UpdateHoverFunction);
+	}
+
+	if (UpdateCoordSystemFunction)
+	{
+		NewGizmo->SetUpdateCoordSystemFunction(UpdateCoordSystemFunction);
+	}
+
+	return NewGizmo;
+}
+
 
 //////////////////////////////////////////////////////////////////////////
 
 
-AXRCreativeBaseTransformGizmoActor::AXRCreativeBaseTransformGizmoActor()
+void UXRCreativeGizmo::Setup()
 {
+	Super::Setup();
+
+	// FIXME: Most ITF sub-gizmos seem to be missing calls to update this
+	UpdateInteractingFunction = [](UPrimitiveComponent* Component, bool bInteracting)
+	{
+		if (UXRCreativeGizmoMeshComponent* XrMeshComp = Cast<UXRCreativeGizmoMeshComponent>(Component))
+		{
+			XrMeshComp->UpdateInteractingState(bInteracting);
+		}
+	};
+
+	//////////////////////////////////////////////////////////////////////////
+	// FIXME?: Have to set these here after base class would otherwise overwrite them
+
+	UpdateHoverFunction = [](UPrimitiveComponent* Component, bool bHovering)
+	{
+		if (UXRCreativeGizmoMeshComponent* XrMeshComp = Cast<UXRCreativeGizmoMeshComponent>(Component))
+		{
+			XrMeshComp->UpdateHoverState(bHovering);
+		}
+	};
+
+	UpdateCoordSystemFunction = [this](UPrimitiveComponent* Component, EToolContextCoordinateSystem CoordSystem)
+	{
+		const bool bWorld = CoordSystem == EToolContextCoordinateSystem::World;
+		if (UGizmoBaseComponent* GizmoBase = Cast<UGizmoBaseComponent>(Component))
+		{
+			GizmoBase->UpdateWorldLocalState(bWorld);
+		}
+		else if (UXRCreativeGizmoMeshComponent* XrMeshComp = Cast<UXRCreativeGizmoMeshComponent>(Component))
+		{
+			FViewCameraState CameraState;
+			GetGizmoManager()->GetContextQueriesAPI()->GetCurrentViewState(CameraState);
+
+			const FTransform Transform = XrMeshComp->CalcViewDependent(CameraState, CoordSystem);
+			XrMeshComp->SetRelativeTransform(Transform);
+		}
+	};
 }
 
 
-UGizmoArrowComponent* AXRCreativeBaseTransformGizmoActor::SetupAxisArrow(UGizmoArrowComponent* InComponent, const FLinearColor& InColor, const FVector& InAxis)
+void UXRCreativeGizmo::Tick(float DeltaTime)
 {
-	InComponent->SetupAttachment(GetRootComponent());
-	InComponent->Direction = InAxis;
-	InComponent->Color = InColor;
-	InComponent->Length = 60.0f;
-	InComponent->Gap = 20.0f;
-	InComponent->Thickness = GizmoLineThickness;
-	InComponent->NotifyExternalPropertyUpdates();
-	InComponent->bHiddenInSceneCapture=true;
-	return InComponent;
+	Super::Tick(DeltaTime);
+
+	if (ACombinedTransformGizmoActor* Actor = GetGizmoActor())
+	{
+		UInteractiveGizmoManager* GizmoManager = GetGizmoManager();
+
+		FViewCameraState CameraState;
+		GizmoManager->GetContextQueriesAPI()->GetCurrentViewState(CameraState);
+
+		const double GizmoDist = (Actor->GetActorLocation() - CameraState.Position).Length();
+
+		const double ConstantScreenSizeBeyondDist = 300.0;
+		const double ScaleFactor = FMath::Max(ConstantScreenSizeBeyondDist, GizmoDist) / ConstantScreenSizeBeyondDist;
+		Actor->GetRootComponent()->SetRelativeScale3D(FVector(ScaleFactor));
+
+		const bool bLocal = GizmoManager->GetContextQueriesAPI()->GetCurrentCoordinateSystem() == EToolContextCoordinateSystem::Local;
+		if (!bLocal)
+		{
+			Actor->GetRootComponent()->SetWorldRotation(FQuat::Identity);
+		}
+	}
 }
 
 
-UGizmoRectangleComponent* AXRCreativeBaseTransformGizmoActor::SetupPlaneRect(UGizmoRectangleComponent* InComponent, const FLinearColor& InColor, const FVector& InAxis0, const FVector& InAxis1)
+UInteractiveGizmo* UXRCreativeGizmo::AddAxisTranslationGizmo(
+	FTransformSubGizmoCommonParams& Params, FTransformSubGizmoSharedState& SharedState)
 {
-	InComponent->SetupAttachment(GetRootComponent());
-	InComponent->DirectionX = InAxis0;
-	InComponent->DirectionY = InAxis1;
-	InComponent->Color = InColor;
-	InComponent->LengthX = InComponent->LengthY = 30.0f;
-	InComponent->SegmentFlags = 0x2 | 0x4;
-	InComponent->Thickness = GizmoLineThickness;
-	InComponent->NotifyExternalPropertyUpdates();
-	InComponent->bHiddenInSceneCapture=true;
-	return InComponent;
+	UAxisPositionGizmo* Gizmo = Cast<UAxisPositionGizmo>(Super::AddAxisTranslationGizmo(Params, SharedState));
+	if (!ensure(Gizmo))
+	{
+		return nullptr;
+	}
+
+	if (UpdateInteractingFunction)
+	{
+		if (UGizmoComponentHitTarget* HitTarget = Cast<UGizmoComponentHitTarget>(Gizmo->HitTarget.GetObject()))
+		{
+			HitTarget->UpdateInteractingFunction = [Component = Params.Component, this](bool bInteracting)
+			{
+				this->UpdateInteractingFunction(Component, bInteracting);
+			};
+
+			HitTarget->UpdateHoverFunction = [Component = Params.Component, this](bool bHovering)
+			{
+				this->UpdateHoverFunction(Component, bHovering);
+			};
+		}
+	}
+
+	if (UXRCreativeGizmoMeshComponent* XrMeshComp = Cast<UXRCreativeGizmoMeshComponent>(Params.Component))
+	{
+		XrMeshComp->AxisSource = Gizmo->AxisSource.GetInterface();
+	}
+
+	return Gizmo;
 }
 
+UInteractiveGizmo* UXRCreativeGizmo::AddPlaneTranslationGizmo(
+	FTransformSubGizmoCommonParams& Params, FTransformSubGizmoSharedState& SharedState)
+{
+	UPlanePositionGizmo* Gizmo = Cast<UPlanePositionGizmo>(Super::AddPlaneTranslationGizmo(Params, SharedState));
+	if (!ensure(Gizmo))
+	{
+		return nullptr;
+	}
 
-UGizmoCircleComponent* AXRCreativeBaseTransformGizmoActor::SetupAxisRotateCircle(UGizmoCircleComponent* InComponent, const FLinearColor& InColor, const FVector& InAxis)
+	if (UpdateInteractingFunction)
+	{
+		if (UGizmoComponentHitTarget* HitTarget = Cast<UGizmoComponentHitTarget>(Gizmo->HitTarget.GetObject()))
+		{
+			HitTarget->UpdateInteractingFunction = [Component = Params.Component, this](bool bInteracting)
+			{
+				this->UpdateInteractingFunction(Component, bInteracting);
+			};
+
+			HitTarget->UpdateHoverFunction = [Component = Params.Component, this](bool bHovering)
+			{
+				this->UpdateHoverFunction(Component, bHovering);
+			};
+		}
+	}
+
+	if (UXRCreativeGizmoMeshComponent* XrMeshComp = Cast<UXRCreativeGizmoMeshComponent>(Params.Component))
+	{
+		XrMeshComp->AxisSource = Gizmo->AxisSource.GetInterface();
+	}
+
+	return Gizmo;
+}
+
+UInteractiveGizmo* UXRCreativeGizmo::AddAxisRotationGizmo(
+	FTransformSubGizmoCommonParams& Params,	FTransformSubGizmoSharedState& SharedState)
 {
-	InComponent->SetupAttachment(GetRootComponent());
-	InComponent->Normal = InAxis;
-	InComponent->Color = InColor;
-	InComponent->Radius = 100.0f;
-	InComponent->Thickness = GizmoLineThickness;
-	InComponent->NotifyExternalPropertyUpdates();
-	InComponent->bHiddenInSceneCapture=true;
-	return InComponent;
+	UAxisAngleGizmo* Gizmo = Cast<UAxisAngleGizmo>(Super::AddAxisRotationGizmo(Params, SharedState));
+	if (!ensure(Gizmo))
+	{
+		return nullptr;
+	}
+
+	if (UpdateInteractingFunction)
+	{
+		if (UGizmoComponentHitTarget* HitTarget = Cast<UGizmoComponentHitTarget>(Gizmo->HitTarget.GetObject()))
+		{
+			HitTarget->UpdateInteractingFunction = [Component = Params.Component, this](bool bInteracting)
+			{
+				this->UpdateInteractingFunction(Component, bInteracting);
+			};
+
+			HitTarget->UpdateHoverFunction = [Component = Params.Component, this](bool bHovering)
+			{
+				this->UpdateHoverFunction(Component, bHovering);
+			};
+		}
+	}
+
+	if (UXRCreativeGizmoMeshComponent* XrMeshComp = Cast<UXRCreativeGizmoMeshComponent>(Params.Component))
+	{
+		XrMeshComp->AxisSource = Gizmo->AxisSource.GetInterface();
+	}
+
+	return Gizmo;
 }
 
+UInteractiveGizmo* UXRCreativeGizmo::AddAxisScaleGizmo(
+	FTransformSubGizmoCommonParams& Params,	FTransformSubGizmoSharedState& SharedState)
+{
+	UAxisPositionGizmo* Gizmo = Cast<UAxisPositionGizmo>(Super::AddAxisScaleGizmo(Params, SharedState));
+	if (!ensure(Gizmo))
+	{
+		return nullptr;
+	}
 
-UGizmoRectangleComponent* AXRCreativeBaseTransformGizmoActor::SetupAxisScaleRect(UGizmoRectangleComponent* InComponent, const FLinearColor& InColor, const FVector& InAxis0, const FVector& InAxis1)
+	if (UpdateInteractingFunction)
+	{
+		if (UGizmoComponentHitTarget* HitTarget = Cast<UGizmoComponentHitTarget>(Gizmo->HitTarget.GetObject()))
+		{
+			HitTarget->UpdateInteractingFunction = [Component = Params.Component, this](bool bInteracting)
+			{
+				this->UpdateInteractingFunction(Component, bInteracting);
+			};
+
+			HitTarget->UpdateHoverFunction = [Component = Params.Component, this](bool bHovering)
+			{
+				this->UpdateHoverFunction(Component, bHovering);
+			};
+		}
+	}
+
+	if (UXRCreativeGizmoMeshComponent* XrMeshComp = Cast<UXRCreativeGizmoMeshComponent>(Params.Component))
+	{
+		XrMeshComp->AxisSource = Gizmo->AxisSource.GetInterface();
+	}
+
+	return Gizmo;
+}
+
+UInteractiveGizmo* UXRCreativeGizmo::AddPlaneScaleGizmo(
+	FTransformSubGizmoCommonParams& Params,	FTransformSubGizmoSharedState& SharedState)
 {
-	InComponent->SetupAttachment(GetRootComponent());
-	InComponent->DirectionX = InAxis0;
-	InComponent->DirectionY = InAxis1;
-	InComponent->Color = InColor;
-	InComponent->OffsetX = 140.0f;
-	InComponent->OffsetY = -10.0f;
-	InComponent->LengthX = 7.0f;
-	InComponent->LengthY = 20.0f;
-	InComponent->Thickness = GizmoLineThickness;
-	InComponent->bOrientYAccordingToCamera = true;
-	InComponent->SegmentFlags = 0x1 | 0x2 | 0x4; // | 0x8;
-	InComponent->NotifyExternalPropertyUpdates();
-	InComponent->bHiddenInSceneCapture=true;
-	return InComponent;
+	UPlanePositionGizmo* Gizmo = Cast<UPlanePositionGizmo>(Super::AddPlaneScaleGizmo(Params, SharedState));
+	if (!ensure(Gizmo))
+	{
+		return nullptr;
+	}
+
+	if (UpdateInteractingFunction)
+	{
+		if (UGizmoComponentHitTarget* HitTarget = Cast<UGizmoComponentHitTarget>(Gizmo->HitTarget.GetObject()))
+		{
+			HitTarget->UpdateInteractingFunction = [Component = Params.Component, this](bool bInteracting)
+			{
+				this->UpdateInteractingFunction(Component, bInteracting);
+			};
+
+			HitTarget->UpdateHoverFunction = [Component = Params.Component, this](bool bHovering)
+			{
+				this->UpdateHoverFunction(Component, bHovering);
+			};
+		}
+	}
+
+	if (UXRCreativeGizmoMeshComponent* XrMeshComp = Cast<UXRCreativeGizmoMeshComponent>(Params.Component))
+	{
+		XrMeshComp->AxisSource = Gizmo->AxisSource.GetInterface();
+	}
+
+	return Gizmo;
+}
+
+UInteractiveGizmo* UXRCreativeGizmo::AddUniformScaleGizmo(
+	FTransformSubGizmoCommonParams& Params,	FTransformSubGizmoSharedState& SharedState)
+{
+	UPlanePositionGizmo* Gizmo = Cast<UPlanePositionGizmo>(Super::AddUniformScaleGizmo(Params, SharedState));
+	if (!ensure(Gizmo))
+	{
+		return nullptr;
+	}
+
+	if (UpdateInteractingFunction)
+	{
+		if (UGizmoComponentHitTarget* HitTarget = Cast<UGizmoComponentHitTarget>(Gizmo->HitTarget.GetObject()))
+		{
+			HitTarget->UpdateInteractingFunction = [Component = Params.Component, this](bool bInteracting)
+			{
+				this->UpdateInteractingFunction(Component, bInteracting);
+			};
+
+			HitTarget->UpdateHoverFunction = [Component = Params.Component, this](bool bHovering)
+			{
+				this->UpdateHoverFunction(Component, bHovering);
+			};
+		}
+	}
+
+	if (UXRCreativeGizmoMeshComponent* XrMeshComp = Cast<UXRCreativeGizmoMeshComponent>(Params.Component))
+	{
+		XrMeshComp->AxisSource = Gizmo->AxisSource.GetInterface();
+	}
+
+	return Gizmo;
+}
+
+
+//////////////////////////////////////////////////////////////////////////
+
+
+void UXRCreativeGizmoMeshComponent::Initialize()
+{
+	const int32 NumMaterials = GetNumMaterials();
+	Materials.Empty(NumMaterials);
+	for (int MaterialIdx = 0; MaterialIdx < NumMaterials; ++MaterialIdx)
+	{
+		UMaterialInterface* BaseMaterial = GetMaterial(MaterialIdx);
+		UMaterialInstanceDynamic* Dynamic = UMaterialInstanceDynamic::Create(BaseMaterial, this);
+		SetMaterial(MaterialIdx, Dynamic);
+		Materials.Add(Dynamic);
+	}
+}
+
+
+void UXRCreativeGizmoMeshComponent::UpdateHoverState(bool bNewHovering)
+{
+	if (bHovering == bNewHovering)
+	{
+		return;
+	}
+
+	bHovering = bNewHovering;
+
+	OnHoveringChanged.Broadcast(this, bNewHovering);
+
+	static const FName HoverMaterialParameterName("GizmoComponentHovering");
+	for (UMaterialInstanceDynamic* Material : Materials)
+	{
+		Material->SetScalarParameterValue(HoverMaterialParameterName, bHovering ? 1.0f : 0.0f);
+	}
+}
+
+
+void UXRCreativeGizmoMeshComponent::UpdateInteractingState(bool bNewInteracting)
+{
+	if (bInteracting == bNewInteracting)
+	{
+		return;
+	}
+
+	bInteracting = bNewInteracting;
+
+	OnInteractingChanged.Broadcast(this, bNewInteracting);
+
+	static const FName InteractingMaterialParameterName("GizmoComponentInteracting");
+	for (UMaterialInstanceDynamic* Material : Materials)
+	{
+		Material->SetScalarParameterValue(InteractingMaterialParameterName, bInteracting ? 1.0f : 0.0f);
+	}
 }
 
 
-UGizmoRectangleComponent* AXRCreativeBaseTransformGizmoActor::SetupPlaneScaleFunc(UGizmoRectangleComponent* InComponent, const FLinearColor& InColor, const FVector& InAxis0, const FVector& InAxis1)
+FTransform UXRCreativeGizmoMeshComponent::CalcViewDependent(
+	const FViewCameraState& InView,
+	EToolContextCoordinateSystem InCoords
+) const
 {
-	InComponent->SetupAttachment(GetRootComponent());
-	InComponent->DirectionX = InAxis0;
-	InComponent->DirectionY = InAxis1;
-	InComponent->Color = InColor;
-	InComponent->OffsetX = InComponent->OffsetY = 120.0f;
-	InComponent->LengthX = InComponent->LengthY = 20.0f;
-	InComponent->Thickness = GizmoLineThickness;
-	InComponent->SegmentFlags = 0x2 | 0x4;
-	InComponent->NotifyExternalPropertyUpdates();
-	InComponent->bHiddenInSceneCapture=true;
-	return InComponent;
+	check(AxisSource);
+
+	UXRCreativeGizmoMeshComponent* Archetype = CastChecked<UXRCreativeGizmoMeshComponent>(GetArchetype());
+
+	FTransform ResultTransform = Archetype->GetRelativeTransform();
+
+	FVector WorldAxisDirection, WorldAxisTangentX, WorldAxisTangentY;
+	FVector CompLocalAxisDirection, CompLocalAxisTangentX, CompLocalAxisTangentY;
+
+	const FTransform& GizmoTransform = GetAttachmentRoot()->GetComponentTransform();
+
+	const bool bWorld = InCoords == EToolContextCoordinateSystem::World;
+	if (bWorld)
+	{
+		AxisSource->GetAxisFrame(CompLocalAxisDirection, CompLocalAxisTangentX, CompLocalAxisTangentY);
+		WorldAxisDirection = CompLocalAxisDirection;
+		WorldAxisTangentX = CompLocalAxisTangentX;
+		WorldAxisTangentY = CompLocalAxisTangentY;
+	}
+	else
+	{
+		AxisSource->GetAxisFrame(WorldAxisDirection, WorldAxisTangentX, WorldAxisTangentY);
+		CompLocalAxisDirection = GizmoTransform.InverseTransformVectorNoScale(WorldAxisDirection);
+		CompLocalAxisTangentX = GizmoTransform.InverseTransformVectorNoScale(WorldAxisTangentX);
+		CompLocalAxisTangentY = GizmoTransform.InverseTransformVectorNoScale(WorldAxisTangentY);
+	}
+
+	const FVector ViewDirection = (GizmoTransform.GetLocation() - InView.Position).GetSafeNormal();
+	const float ViewDotAxisDirection = ViewDirection.Dot(WorldAxisDirection);
+	const float ViewDotAxisTangentX = ViewDirection.Dot(WorldAxisTangentX);
+	const float ViewDotAxisTangentY = ViewDirection.Dot(WorldAxisTangentY);
+
+	if (HideAbsoluteViewDotThreshold > 0.0f && FMath::Abs(ViewDotAxisDirection) <= HideAbsoluteViewDotThreshold)
+	{
+		ResultTransform.SetScale3D(FVector::ZeroVector);
+		return ResultTransform;
+	}
+
+	FVector ReflectionScale = FVector::OneVector;
+
+	if (bReflectOnPrimaryAxis && ViewDotAxisDirection > 0.0f)
+	{
+		ReflectionScale += CompLocalAxisDirection * -2.0;
+	}
+
+	if (bReflectOnTangentAxes)
+	{
+		if (ViewDotAxisTangentX > 0.0f)
+		{
+			ReflectionScale += CompLocalAxisTangentX * -2.0;
+		}
+
+		if (ViewDotAxisTangentY > 0.0f)
+		{
+			ReflectionScale += CompLocalAxisTangentY * -2.0;
+		}
+	}
+
+	const FTransform ReflectScale = FTransform(FQuat::Identity, FVector::ZeroVector, ReflectionScale);
+	ResultTransform = ResultTransform * ReflectScale;
+
+	return ResultTransform;
+}
+
+
+//////////////////////////////////////////////////////////////////////////
+
+
+AXRCreativeCombinedTransformGizmoActor::AXRCreativeCombinedTransformGizmoActor()
+{
+	PrimaryActorTick.bCanEverTick = true;
+
+	XRSceneRoot = RootComponent;
+
+	TranslateX = XRTranslateX = CreateOptionalDefaultSubobject<UXRCreativeGizmoMeshComponent>("TranslateX");
+	TranslateX->SetupAttachment(GetRootComponent());
+	TranslateY = XRTranslateY = CreateOptionalDefaultSubobject<UXRCreativeGizmoMeshComponent>("TranslateY");
+	TranslateY->SetupAttachment(GetRootComponent());
+	TranslateZ = XRTranslateZ = CreateOptionalDefaultSubobject<UXRCreativeGizmoMeshComponent>("TranslateZ");
+	TranslateZ->SetupAttachment(GetRootComponent());
+	TranslateYZ = XRTranslateYZ = CreateOptionalDefaultSubobject<UXRCreativeGizmoMeshComponent>("TranslateYZ");
+	TranslateYZ->SetupAttachment(GetRootComponent());
+	TranslateXZ = XRTranslateXZ = CreateOptionalDefaultSubobject<UXRCreativeGizmoMeshComponent>("TranslateXZ");
+	TranslateXZ->SetupAttachment(GetRootComponent());
+	TranslateXY = XRTranslateXY = CreateOptionalDefaultSubobject<UXRCreativeGizmoMeshComponent>("TranslateXY");
+	TranslateXY->SetupAttachment(GetRootComponent());
+
+	// TODO: RotationSphere?
+	RotateX = XRRotateX = CreateOptionalDefaultSubobject<UXRCreativeGizmoMeshComponent>("RotateX");
+	RotateX->SetupAttachment(GetRootComponent());
+	RotateY = XRRotateY = CreateOptionalDefaultSubobject<UXRCreativeGizmoMeshComponent>("RotateY");
+	RotateY->SetupAttachment(GetRootComponent());
+	RotateZ = XRRotateZ = CreateOptionalDefaultSubobject<UXRCreativeGizmoMeshComponent>("RotateZ");
+	RotateZ->SetupAttachment(GetRootComponent());
+
+	UniformScale = XRUniformScale = CreateOptionalDefaultSubobject<UXRCreativeGizmoMeshComponent>("UniformScale");
+	UniformScale->SetupAttachment(GetRootComponent());
+	AxisScaleX = XRAxisScaleX = CreateOptionalDefaultSubobject<UXRCreativeGizmoMeshComponent>("AxisScaleX");
+	AxisScaleX->SetupAttachment(GetRootComponent());
+	AxisScaleY = XRAxisScaleY = CreateOptionalDefaultSubobject<UXRCreativeGizmoMeshComponent>("AxisScaleY");
+	AxisScaleY->SetupAttachment(GetRootComponent());
+	AxisScaleZ = XRAxisScaleZ = CreateOptionalDefaultSubobject<UXRCreativeGizmoMeshComponent>("AxisScaleZ");
+	AxisScaleZ->SetupAttachment(GetRootComponent());
+	PlaneScaleYZ = XRPlaneScaleYZ = CreateOptionalDefaultSubobject<UXRCreativeGizmoMeshComponent>("AxisScaleYZ");
+	PlaneScaleYZ->SetupAttachment(GetRootComponent());
+	PlaneScaleXZ = XRPlaneScaleXZ = CreateOptionalDefaultSubobject<UXRCreativeGizmoMeshComponent>("AxisScaleXZ");
+	PlaneScaleXZ->SetupAttachment(GetRootComponent());
+	PlaneScaleXY = XRPlaneScaleXY = CreateOptionalDefaultSubobject<UXRCreativeGizmoMeshComponent>("AxisScaleXY");
+	PlaneScaleXY->SetupAttachment(GetRootComponent());
 }
 
 
-void AXRCreativeBaseTransformGizmoActor::ConstructDefaults(ETransformGizmoSubElements EnableElements)
+void AXRCreativeCombinedTransformGizmoActor::Tick(float DeltaSeconds)
 {
-	if ((EnableElements & ETransformGizmoSubElements::TranslateAxisX) != ETransformGizmoSubElements::None)
+	Super::Tick(DeltaSeconds);
+}
+
+
+void AXRCreativeCombinedTransformGizmoActor::SetEnabledElements(ETransformGizmoSubElements EnableElements)
+{
+	const bool bEnableTranslateAxisX = (EnableElements & ETransformGizmoSubElements::TranslateAxisX) != ETransformGizmoSubElements::None;
+	const bool bEnableTranslateAxisY = (EnableElements & ETransformGizmoSubElements::TranslateAxisY) != ETransformGizmoSubElements::None;
+	const bool bEnableTranslateAxisZ = (EnableElements & ETransformGizmoSubElements::TranslateAxisZ) != ETransformGizmoSubElements::None;
+	if (TranslateX)
 	{
-		TranslateX = SetupAxisArrow(CreateDefaultSubobject<UGizmoArrowComponent>("TranslateX"), FLinearColor::Red, FVector(1, 0, 0));
+		TranslateX->SetVisibility(bEnableTranslateAxisX, true);
 	}
-	if ((EnableElements & ETransformGizmoSubElements::TranslateAxisY) != ETransformGizmoSubElements::None)
+	if (TranslateY)
 	{
-		TranslateY = SetupAxisArrow(CreateDefaultSubobject<UGizmoArrowComponent>("TranslateY"), FLinearColor::Green, FVector(0, 1, 0));
+		TranslateY->SetVisibility(bEnableTranslateAxisY, true);
 	}
-	if ((EnableElements & ETransformGizmoSubElements::TranslateAxisZ) != ETransformGizmoSubElements::None)
+	if (TranslateZ)
 	{
-		TranslateZ = SetupAxisArrow(CreateDefaultSubobject<UGizmoArrowComponent>("TranslateZ"), FLinearColor::Blue, FVector(0, 0, 1));
+		TranslateZ->SetVisibility(bEnableTranslateAxisZ, true);
 	}
 
-	if ((EnableElements & ETransformGizmoSubElements::TranslatePlaneYZ) != ETransformGizmoSubElements::None)
+	const bool bEnableTransplatePlaneYZ = (EnableElements & ETransformGizmoSubElements::TranslatePlaneYZ) != ETransformGizmoSubElements::None;
+	const bool bEnableTransplatePlaneXZ = (EnableElements & ETransformGizmoSubElements::TranslatePlaneXZ) != ETransformGizmoSubElements::None;
+	const bool bEnableTransplatePlaneXY = (EnableElements & ETransformGizmoSubElements::TranslatePlaneXY) != ETransformGizmoSubElements::None;
+	if (TranslateYZ)
 	{
-		TranslateYZ = SetupPlaneRect(CreateDefaultSubobject<UGizmoRectangleComponent>("TranslateYZ"), FLinearColor::Red, FVector(0, 1, 0), FVector(0, 0, 1));
+		TranslateYZ->SetVisibility(bEnableTransplatePlaneYZ, true);
 	}
-	if ((EnableElements & ETransformGizmoSubElements::TranslatePlaneXZ) != ETransformGizmoSubElements::None)
+	if (TranslateXZ)
 	{
-		TranslateXZ = SetupPlaneRect(CreateDefaultSubobject<UGizmoRectangleComponent>("TranslateXZ"), FLinearColor::Green, FVector(1, 0, 0), FVector(0, 0, 1));
+		TranslateXZ->SetVisibility(bEnableTransplatePlaneXZ, true);
 	}
-	if ((EnableElements & ETransformGizmoSubElements::TranslatePlaneXY) != ETransformGizmoSubElements::None)
+	if (TranslateXY)
 	{
-		TranslateXY = SetupPlaneRect(CreateDefaultSubobject<UGizmoRectangleComponent>("TranslateXY"), FLinearColor::Blue, FVector(1, 0, 0), FVector(0, 1, 0));
+		TranslateXY->SetVisibility(bEnableTransplatePlaneXY, true);
 	}
 
-	bool bAnyRotate = false;
-	if ((EnableElements & ETransformGizmoSubElements::RotateAxisX) != ETransformGizmoSubElements::None)
+	const bool bEnableRotateAxisX = (EnableElements & ETransformGizmoSubElements::RotateAxisX) != ETransformGizmoSubElements::None;
+	const bool bEnableRotateAxisY = (EnableElements & ETransformGizmoSubElements::RotateAxisY) != ETransformGizmoSubElements::None;
+	const bool bEnableRotateAxisZ = (EnableElements & ETransformGizmoSubElements::RotateAxisZ) != ETransformGizmoSubElements::None;
+	if (RotateX)
 	{
-		RotateX = SetupAxisRotateCircle(CreateDefaultSubobject<UGizmoCircleComponent>("RotateX"), FLinearColor::Red, FVector(1, 0, 0));
-		bAnyRotate = true;
+		RotateX->SetVisibility(bEnableRotateAxisX, true);
 	}
-	if ((EnableElements & ETransformGizmoSubElements::RotateAxisY) != ETransformGizmoSubElements::None)
+	if (RotateY)
 	{
-		RotateY = SetupAxisRotateCircle(CreateDefaultSubobject<UGizmoCircleComponent>("RotateY"), FLinearColor::Green, FVector(0, 1, 0));
-		bAnyRotate = true;
+		RotateY->SetVisibility(bEnableRotateAxisY, true);
 	}
-	if ((EnableElements & ETransformGizmoSubElements::RotateAxisZ) != ETransformGizmoSubElements::None)
+	if (RotateZ)
 	{
-		RotateZ = SetupAxisRotateCircle(CreateDefaultSubobject<UGizmoCircleComponent>("RotateZ"), FLinearColor::Blue, FVector(0, 0, 1));
-		bAnyRotate = true;
+		RotateZ->SetVisibility(bEnableRotateAxisZ, true);
 	}
 
+#if 0
 	// add a non-interactive view-aligned circle element, so the axes look like a sphere.
+	const bool bAnyRotate = bEnableRotateAxisX || bEnableRotateAxisY || bEnableRotateAxisZ;
 	if (bAnyRotate)
 	{
-		UGizmoCircleComponent* RotationSphereComp = CreateDefaultSubobject<UGizmoCircleComponent>("RotationSphere");
+		UGizmoCircleComponent* RotationSphereComp = XR	UGizmoCircleComponent* RotationSphereComp = CreateOptionalDefaultSubobject<UGizmoCircleComponent>("RotationSphere");
 		RotationSphereComp->SetupAttachment(GetRootComponent());
 		RotationSphereComp->Color = FLinearColor::Gray;
 		RotationSphereComp->Thickness = GizmoLineThickness * 0.5f;
 		RotationSphereComp->Radius = 120.0f;
 		RotationSphereComp->bViewAligned = true;
 		RotationSphereComp->NotifyExternalPropertyUpdates();
-		RotationSphereComp->bHiddenInSceneCapture=true;
 		RotationSphere = RotationSphereComp;
 	}
+#endif
 
-	if ((EnableElements & ETransformGizmoSubElements::ScaleUniform) != ETransformGizmoSubElements::None)
+	const bool bEnableScaleUniform = (EnableElements & ETransformGizmoSubElements::ScaleUniform) != ETransformGizmoSubElements::None;
+	const bool bEnableScaleAxisX = (EnableElements & ETransformGizmoSubElements::ScaleAxisX) != ETransformGizmoSubElements::None;
+	const bool bEnableScaleAxisY = (EnableElements & ETransformGizmoSubElements::ScaleAxisY) != ETransformGizmoSubElements::None;
+	const bool bEnableScaleAxisZ = (EnableElements & ETransformGizmoSubElements::ScaleAxisZ) != ETransformGizmoSubElements::None;
+	if (UniformScale)
 	{
-		float BoxSize = 14.0f;
-		UGizmoBoxComponent* UniformScaleComp = CreateDefaultSubobject<UGizmoBoxComponent>("UniformScale");
-		UniformScaleComp->SetupAttachment(GetRootComponent());
-		UniformScaleComp->Color = FLinearColor::Black;
-		UniformScaleComp->Origin = FVector(BoxSize/2, BoxSize/2, BoxSize/2);
-		UniformScaleComp->Dimensions = FVector(BoxSize, BoxSize, BoxSize);
-		UniformScaleComp->NotifyExternalPropertyUpdates();
-		UniformScaleComp->bHiddenInSceneCapture=true;
-		UniformScale = UniformScaleComp;
+		UniformScale->SetVisibility(bEnableScaleUniform, true);
 	}
-
-	if ((EnableElements & ETransformGizmoSubElements::ScaleAxisX) != ETransformGizmoSubElements::None)
+	if (AxisScaleX)
 	{
-		AxisScaleX = SetupAxisScaleRect(CreateDefaultSubobject<UGizmoRectangleComponent>("AxisScaleX"), FLinearColor::Red, FVector(1, 0, 0), FVector(0, 0, 1));
+		AxisScaleX->SetVisibility(bEnableScaleAxisX, true);
 	}
-	if ((EnableElements & ETransformGizmoSubElements::ScaleAxisY) != ETransformGizmoSubElements::None)
+	if (AxisScaleY)
 	{
-		AxisScaleY = SetupAxisScaleRect(CreateDefaultSubobject<UGizmoRectangleComponent>("AxisScaleY"), FLinearColor::Green, FVector(0, 1, 0), FVector(0, 0, 1));
+		AxisScaleY->SetVisibility(bEnableScaleAxisY, true);
 	}
-	if ((EnableElements & ETransformGizmoSubElements::ScaleAxisZ) != ETransformGizmoSubElements::None)
+	if (AxisScaleZ)
 	{
-		AxisScaleZ = SetupAxisScaleRect(CreateDefaultSubobject<UGizmoRectangleComponent>("AxisScaleZ"), FLinearColor::Blue, FVector(0, 0, 1), FVector(1, 0, 0));
+		AxisScaleZ->SetVisibility(bEnableScaleAxisZ, true);
 	}
 
-	if ((EnableElements & ETransformGizmoSubElements::ScalePlaneYZ) != ETransformGizmoSubElements::None)
+	const bool bEnableScalePlaneYZ = (EnableElements & ETransformGizmoSubElements::ScalePlaneYZ) != ETransformGizmoSubElements::None;
+	const bool bEnableScalePlaneXZ = (EnableElements & ETransformGizmoSubElements::ScalePlaneXZ) != ETransformGizmoSubElements::None;
+	const bool bEnableScalePlaneXY = (EnableElements & ETransformGizmoSubElements::ScalePlaneXY) != ETransformGizmoSubElements::None;
+	if (PlaneScaleYZ)
 	{
-		PlaneScaleYZ = SetupPlaneScaleFunc(CreateDefaultSubobject<UGizmoRectangleComponent>("PlaneScaleYZ"), FLinearColor::Red, FVector(0, 1, 0), FVector(0, 0, 1));
+		PlaneScaleYZ->SetVisibility(bEnableScalePlaneYZ, true);
 	}
-	if ((EnableElements & ETransformGizmoSubElements::ScalePlaneXZ) != ETransformGizmoSubElements::None)
+	if (PlaneScaleXZ)
 	{
-		PlaneScaleXZ = SetupPlaneScaleFunc(CreateDefaultSubobject<UGizmoRectangleComponent>("PlaneScaleXZ"), FLinearColor::Green, FVector(1, 0, 0), FVector(0, 0, 1));
+		PlaneScaleXZ->SetVisibility(bEnableScalePlaneXZ, true);
 	}
-	if ((EnableElements & ETransformGizmoSubElements::ScalePlaneXY) != ETransformGizmoSubElements::None)
+	if (PlaneScaleXY)
 	{
-		PlaneScaleXY = SetupPlaneScaleFunc(CreateDefaultSubobject<UGizmoRectangleComponent>("PlaneScaleXY"), FLinearColor::Blue, FVector(1, 0, 0), FVector(0, 1, 0));
+		PlaneScaleXY->SetVisibility(bEnableScalePlaneXY, true);
 	}
 }
diff --git a/Engine/Plugins/Experimental/XRCreativeFramework/Source/XRCreative/Private/ITF/XRCreativeITFComponent.cpp b/Engine/Plugins/Experimental/XRCreativeFramework/Source/XRCreative/Private/ITF/XRCreativeITFComponent.cpp
--- a/Engine/Plugins/Experimental/XRCreativeFramework/Source/XRCreative/Private/ITF/XRCreativeITFComponent.cpp
+++ b/Engine/Plugins/Experimental/XRCreativeFramework/Source/XRCreative/Private/ITF/XRCreativeITFComponent.cpp
@@ -3,6 +3,7 @@
 #include "XRCreativeITFComponent.h"
 #include "XRCreativeAvatar.h"
 #include "XRCreativeITFRenderComponent.h"
+#include "XRCreativeGizmos.h"
 #include "XRCreativeLog.h"
 #include "XRCreativePointerComponent.h"
 #include "SelectionInteraction.h"
@@ -353,6 +354,54 @@ protected:
 //////////////////////////////////////////////////////////////////////////
 
 
+class FXRCreativeTransformGizmoActorFactory : public FCombinedTransformGizmoActorFactory
+{
+public:
+	FXRCreativeTransformGizmoActorFactory(
+		UGizmoViewContext* InGizmoViewContext,
+		UXRCreativeITFComponent* InITFComponent
+	)
+		: FCombinedTransformGizmoActorFactory(InGizmoViewContext)
+		, ITFComponent(InITFComponent)
+	{
+		ensure(InGizmoViewContext);
+		ensure(InITFComponent);
+	}
+
+	virtual ACombinedTransformGizmoActor* CreateNewGizmoActor(UWorld* World) const override
+	{
+		TSubclassOf<AXRCreativeCombinedTransformGizmoActor> GizmoClass =
+			ITFComponent->GetCombinedGizmoActorClass();
+
+		AXRCreativeCombinedTransformGizmoActor* Gizmo =
+			World->SpawnActor<AXRCreativeCombinedTransformGizmoActor>(GizmoClass);
+
+		if (!GizmoClass || !ensure(Gizmo))
+		{
+			// Fall back to base
+			return FCombinedTransformGizmoActorFactory::CreateNewGizmoActor(World);
+		}
+
+		Gizmo->SetEnabledElements(EnableElements);
+
+		TInlineComponentArray<UXRCreativeGizmoMeshComponent*> GizmoMeshes;
+		Gizmo->GetComponents(GizmoMeshes);
+		for (UXRCreativeGizmoMeshComponent* GizmoMesh : GizmoMeshes)
+		{
+			GizmoMesh->Initialize();
+		}
+
+		return Gizmo;
+	}
+
+protected:
+	UXRCreativeITFComponent* ITFComponent;
+};
+
+
+//////////////////////////////////////////////////////////////////////////
+
+
 UXRCreativeITFComponent::UXRCreativeITFComponent()
 {
 	PrimaryComponentTick.bCanEverTick = true;
@@ -461,8 +510,17 @@ void UXRCreativeITFComponent::InitializeComponent()
 	}
 
 	// create transform interaction
+	UGizmoViewContext* GizmoViewContext =
+		ToolsContext->GizmoManager->GetContextObjectStore()->FindContext<UGizmoViewContext>();
+
+	TSharedRef<FXRCreativeTransformGizmoActorFactory> GizmoActorFactory =
+		MakeShared<FXRCreativeTransformGizmoActorFactory>(GizmoViewContext, this);
+
 	TransformInteraction = NewObject<UXRCreativeTransformInteraction>(this);
-	TransformInteraction->Initialize(GetSelectionSet(), ToolsContext->GizmoManager,
+	TransformInteraction->Initialize(
+		GizmoActorFactory,
+		GetSelectionSet(),
+		ToolsContext->GizmoManager,
 		[this]() { return HaveActiveTool() == false; }
 	);
 	TransformInteraction->ForceUpdateGizmoState();
diff --git a/Engine/Plugins/Experimental/XRCreativeFramework/Source/XRCreative/Private/XRCreativeAvatar.cpp b/Engine/Plugins/Experimental/XRCreativeFramework/Source/XRCreative/Private/XRCreativeAvatar.cpp
--- a/Engine/Plugins/Experimental/XRCreativeFramework/Source/XRCreative/Private/XRCreativeAvatar.cpp
+++ b/Engine/Plugins/Experimental/XRCreativeFramework/Source/XRCreative/Private/XRCreativeAvatar.cpp
@@ -183,6 +183,14 @@ void AXRCreativeAvatar::BeginPlay()
 }
 
 
+void AXRCreativeAvatar::GetActorEyesViewPoint(FVector& Location, FRotator& Rotation) const
+{
+	FTransform HeadTransform = GetHeadTransform();
+	Location = HeadTransform.GetLocation();
+	Rotation = HeadTransform.GetRotation().Rotator();
+}
+
+
 void AXRCreativeAvatar::ConfigureToolset(UXRCreativeToolset* InToolset)
 {
 	check(InToolset);
diff --git a/Engine/Plugins/Experimental/XRCreativeFramework/Source/XRCreative/Public/XRCreativeAvatar.h b/Engine/Plugins/Experimental/XRCreativeFramework/Source/XRCreative/Public/XRCreativeAvatar.h
--- a/Engine/Plugins/Experimental/XRCreativeFramework/Source/XRCreative/Public/XRCreativeAvatar.h
+++ b/Engine/Plugins/Experimental/XRCreativeFramework/Source/XRCreative/Public/XRCreativeAvatar.h
@@ -48,6 +48,7 @@ public:
 	virtual void BeginDestroy() override;
 	virtual void Tick(float InDeltaSeconds) override;
 	virtual void BeginPlay() override;
+	virtual void GetActorEyesViewPoint(FVector& Location, FRotator& Rotation) const override;
 
 	void ConfigureToolset(UXRCreativeToolset* InToolset);
 	const UXRCreativeToolset* GetToolset() const { return Toolset; }
diff --git a/Engine/Plugins/Experimental/XRCreativeFramework/Source/XRCreative/Public/XRCreativeGizmos.h b/Engine/Plugins/Experimental/XRCreativeFramework/Source/XRCreative/Public/XRCreativeGizmos.h
--- a/Engine/Plugins/Experimental/XRCreativeFramework/Source/XRCreative/Public/XRCreativeGizmos.h
+++ b/Engine/Plugins/Experimental/XRCreativeFramework/Source/XRCreative/Public/XRCreativeGizmos.h
@@ -3,9 +3,18 @@
 #pragma once
 
 #include "BaseGizmos/CombinedTransformGizmo.h"
+#include "Components/StaticMeshComponent.h"
 #include "XRCreativeGizmos.generated.h"
 
 
+namespace UE::GizmoUtil
+{
+	struct FTransformSubGizmoCommonParams;
+	struct FTransformSubGizmoSharedState;
+}
+
+
+/** Responsible for instantiating our UXRCreativeGizmo subclass. */
 UCLASS()
 class UXRCreativeGizmoBuilder : public UCombinedTransformGizmoBuilder
 {
@@ -13,65 +22,192 @@ class UXRCreativeGizmoBuilder : public UCombinedTransformGizmoBuilder
 
 public:
 	UXRCreativeGizmoBuilder();
+
+	virtual UInteractiveGizmo* BuildGizmo(const FToolBuilderState& SceneState) const override;
 };
 
 
-UCLASS(Blueprintable)
-class AXRCreativeBaseTransformGizmoActor : public ACombinedTransformGizmoActor
+/** Implements our UpdateHoverFunction, and introduces an analogous UpdateInteractingFunction. */
+UCLASS()
+class UXRCreativeGizmo : public UCombinedTransformGizmo
 {
 	GENERATED_BODY()
 
 public:
-	AXRCreativeBaseTransformGizmoActor();
+	virtual void Setup() override;
+	virtual void Tick(float DeltaTime) override;
 
 protected:
-	// Only used during Setup(...) component methods; set before calling those or ConstructDefaults
-	float GizmoLineThickness = 3.0f;
+	using FTransformSubGizmoCommonParams = UE::GizmoUtil::FTransformSubGizmoCommonParams;
+	using FTransformSubGizmoSharedState = UE::GizmoUtil::FTransformSubGizmoSharedState;
 
-	UGizmoArrowComponent* SetupAxisArrow(UGizmoArrowComponent* InComponent, const FLinearColor& InColor, const FVector& InAxis);
-	UGizmoRectangleComponent* SetupPlaneRect(UGizmoRectangleComponent* InComponent, const FLinearColor& InColor, const FVector& InAxis0, const FVector& InAxis1);
-	UGizmoCircleComponent* SetupAxisRotateCircle(UGizmoCircleComponent* InComponent, const FLinearColor& InColor, const FVector& InAxis);
-	UGizmoRectangleComponent* SetupAxisScaleRect(UGizmoRectangleComponent* InComponent, const FLinearColor& InColor, const FVector& InAxis0, const FVector& InAxis1);
-	UGizmoRectangleComponent* SetupPlaneScaleFunc(UGizmoRectangleComponent* InComponent, const FLinearColor& InColor, const FVector& InAxis0, const FVector& InAxis1);
+	virtual UInteractiveGizmo* AddAxisTranslationGizmo(FTransformSubGizmoCommonParams& Params, FTransformSubGizmoSharedState& SharedState) override;
+	virtual UInteractiveGizmo* AddPlaneTranslationGizmo(FTransformSubGizmoCommonParams& Params, FTransformSubGizmoSharedState& SharedState) override;
+	virtual UInteractiveGizmo* AddAxisRotationGizmo(FTransformSubGizmoCommonParams& Params, FTransformSubGizmoSharedState& SharedState) override;
+	virtual UInteractiveGizmo* AddAxisScaleGizmo(FTransformSubGizmoCommonParams& Params, FTransformSubGizmoSharedState& SharedState) override;
+	virtual UInteractiveGizmo* AddPlaneScaleGizmo(FTransformSubGizmoCommonParams& Params, FTransformSubGizmoSharedState& SharedState) override;
+	virtual UInteractiveGizmo* AddUniformScaleGizmo(FTransformSubGizmoCommonParams& Params, FTransformSubGizmoSharedState& SharedState) override;
 
-	void ConstructDefaults(ETransformGizmoSubElements EnableElements);
+protected:
+	TFunction<void(UPrimitiveComponent*, bool)> UpdateInteractingFunction;
 };
 
 
-UCLASS()
-class AXRCreativeTRSGizmoActor : public AXRCreativeBaseTransformGizmoActor
+class UXRCreativeGizmoMeshComponent;
+
+
+UCLASS(Blueprintable, NonTransient)
+class AXRCreativeCombinedTransformGizmoActor : public ACombinedTransformGizmoActor
 {
 	GENERATED_BODY()
 
 public:
-	AXRCreativeTRSGizmoActor()
-	{
-		ConstructDefaults(ETransformGizmoSubElements::FullTranslateRotateScale);
-	}
-};
+	AXRCreativeCombinedTransformGizmoActor();
 
+	virtual void Tick(float DeltaSeconds) override;
+	virtual bool ShouldTickIfViewportsOnly() const override { return true; }
 
-UCLASS()
-class AXRCreativeTRUSGizmoActor : public AXRCreativeBaseTransformGizmoActor
-{
-	GENERATED_BODY()
+	void SetEnabledElements(ETransformGizmoSubElements EnableElements);
+
+	UPROPERTY(BlueprintReadOnly, Category="Gizmo")
+	TObjectPtr<USceneComponent> WorldAligned;
+
+	TWeakObjectPtr<UInteractiveGizmoManager> WeakGizmoManager;
 
 public:
-	AXRCreativeTRUSGizmoActor()
-	{
-		ConstructDefaults(ETransformGizmoSubElements::TranslateRotateUniformScale);
-	}
+	// NOTE: These properties alias ones inherited from ACombinedTransformGizmoActor
+	// as a workaround to expose them as EditAnywhere + BlueprintReadWrite without
+	// modifying the base class.
+
+	//////////////////////////////////////////////////////////////////////////
+
+	UPROPERTY(EditAnywhere, BlueprintReadWrite, Category="Gizmo")
+	TObjectPtr<USceneComponent> XRSceneRoot;
+
+	//
+	// Translation Components
+	//
+
+	/** X Axis Translation Component */
+	UPROPERTY(EditAnywhere, BlueprintReadWrite, Category="Gizmo")
+	TObjectPtr<UXRCreativeGizmoMeshComponent> XRTranslateX;
+
+	/** Y Axis Translation Component */
+	UPROPERTY(EditAnywhere, BlueprintReadWrite, Category="Gizmo")
+	TObjectPtr<UXRCreativeGizmoMeshComponent> XRTranslateY;
+
+	/** Z Axis Translation Component */
+	UPROPERTY(EditAnywhere, BlueprintReadWrite, Category="Gizmo")
+	TObjectPtr<UXRCreativeGizmoMeshComponent> XRTranslateZ;
+
+
+	/** YZ Plane Translation Component */
+	UPROPERTY(EditAnywhere, BlueprintReadWrite, Category="Gizmo")
+	TObjectPtr<UXRCreativeGizmoMeshComponent> XRTranslateYZ;
+
+	/** XZ Plane Translation Component */
+	UPROPERTY(EditAnywhere, BlueprintReadWrite, Category="Gizmo")
+	TObjectPtr<UXRCreativeGizmoMeshComponent> XRTranslateXZ;
+
+	/** XY Plane Translation Component */
+	UPROPERTY(EditAnywhere, BlueprintReadWrite, Category="Gizmo")
+	TObjectPtr<UXRCreativeGizmoMeshComponent> XRTranslateXY;
+
+	//
+	// Rotation Components
+	//
+
+	/** X Axis Rotation Component */
+	UPROPERTY(EditAnywhere, BlueprintReadWrite, Category="Gizmo")
+	TObjectPtr<UXRCreativeGizmoMeshComponent> XRRotateX;
+
+	/** Y Axis Rotation Component */
+	UPROPERTY(EditAnywhere, BlueprintReadWrite, Category="Gizmo")
+	TObjectPtr<UXRCreativeGizmoMeshComponent> XRRotateY;
+
+	/** Z Axis Rotation Component */
+	UPROPERTY(EditAnywhere, BlueprintReadWrite, Category="Gizmo")
+	TObjectPtr<UXRCreativeGizmoMeshComponent> XRRotateZ;
+
+	/** Z Axis Rotation Component */
+	UPROPERTY(EditAnywhere, BlueprintReadWrite, Category="Gizmo")
+	TObjectPtr<UXRCreativeGizmoMeshComponent> XRRotationSphere;
+
+	//
+	// Scaling Components
+	//
+
+	/** Uniform Scale Component */
+	UPROPERTY(EditAnywhere, BlueprintReadWrite, Category="Gizmo")
+	TObjectPtr<UXRCreativeGizmoMeshComponent> XRUniformScale;
+
+
+	/** X Axis Scale Component */
+	UPROPERTY(EditAnywhere, BlueprintReadWrite, Category="Gizmo")
+	TObjectPtr<UXRCreativeGizmoMeshComponent> XRAxisScaleX;
+
+	/** Y Axis Scale Component */
+	UPROPERTY(EditAnywhere, BlueprintReadWrite, Category="Gizmo")
+	TObjectPtr<UXRCreativeGizmoMeshComponent> XRAxisScaleY;
+
+	/** Z Axis Scale Component */
+	UPROPERTY(EditAnywhere, BlueprintReadWrite, Category="Gizmo")
+	TObjectPtr<UXRCreativeGizmoMeshComponent> XRAxisScaleZ;
+
+
+	/** YZ Plane Scale Component */
+	UPROPERTY(EditAnywhere, BlueprintReadWrite, Category="Gizmo")
+	TObjectPtr<UXRCreativeGizmoMeshComponent> XRPlaneScaleYZ;
+
+	/** XZ Plane Scale Component */
+	UPROPERTY(EditAnywhere, BlueprintReadWrite, Category="Gizmo")
+	TObjectPtr<UXRCreativeGizmoMeshComponent> XRPlaneScaleXZ;
+
+	/** XY Plane Scale Component */
+	UPROPERTY(EditAnywhere, BlueprintReadWrite, Category="Gizmo")
+	TObjectPtr<UXRCreativeGizmoMeshComponent> XRPlaneScaleXY;
 };
 
 
+DECLARE_DYNAMIC_MULTICAST_DELEGATE_TwoParams(FXRCreativeGizmoStateChanged, UXRCreativeGizmoMeshComponent*, Component, bool, bNewState);
+
+
 UCLASS()
-class AXRCreativeTRGizmoActor : public AXRCreativeBaseTransformGizmoActor
+class UXRCreativeGizmoMeshComponent : public UStaticMeshComponent
 {
 	GENERATED_BODY()
 
 public:
-	AXRCreativeTRGizmoActor()
-	{
-		ConstructDefaults(ETransformGizmoSubElements::StandardTranslateRotate);
-	}
+	void Initialize();
+	void UpdateHoverState(bool bInHovering);
+	void UpdateInteractingState(bool bInInteracting);
+
+	UPROPERTY(BlueprintAssignable)
+	FXRCreativeGizmoStateChanged OnHoveringChanged;
+
+	UPROPERTY(BlueprintAssignable)
+	FXRCreativeGizmoStateChanged OnInteractingChanged;
+
+	FTransform CalcViewDependent(const FViewCameraState& InView, EToolContextCoordinateSystem InCoords) const;
+
+	IGizmoAxisSource* AxisSource = nullptr;
+
+protected:
+	UPROPERTY(BlueprintReadOnly, Category="Gizmo")
+	bool bHovering;
+
+	UPROPERTY(BlueprintReadOnly, Category="Gizmo")
+	bool bInteracting;
+
+	UPROPERTY(EditAnywhere, BlueprintReadWrite, Category="Gizmo")
+	float HideAbsoluteViewDotThreshold = -1.0f;
+
+	UPROPERTY(EditAnywhere, BlueprintReadWrite, Category="Gizmo")
+	bool bReflectOnPrimaryAxis;
+
+	UPROPERTY(EditAnywhere, BlueprintReadWrite, Category="Gizmo")
+	bool bReflectOnTangentAxes;
+
+	UPROPERTY(Transient)
+	TArray<TObjectPtr<UMaterialInstanceDynamic>> Materials;
 };
diff --git a/Engine/Plugins/Experimental/XRCreativeFramework/Source/XRCreative/Public/XRCreativeITFComponent.h b/Engine/Plugins/Experimental/XRCreativeFramework/Source/XRCreative/Public/XRCreativeITFComponent.h
--- a/Engine/Plugins/Experimental/XRCreativeFramework/Source/XRCreative/Public/XRCreativeITFComponent.h
+++ b/Engine/Plugins/Experimental/XRCreativeFramework/Source/XRCreative/Public/XRCreativeITFComponent.h
@@ -9,6 +9,7 @@
 
 
 class AXRCreativeBaseTransformGizmoActor;
+class AXRCreativeCombinedTransformGizmoActor;
 class FXRCreativeToolsContextTransactionImpl;
 class FXRCreativeToolsContextQueriesImpl;
 class UInteractiveToolsContext;
@@ -92,9 +93,12 @@ public:
 	UFUNCTION(BlueprintCallable, Category="XR Creative|Tools")
 	void SetCurrentTransformGizmoMode(EToolContextTransformGizmoMode GizmoMode);
 
+
+	TSubclassOf<AXRCreativeCombinedTransformGizmoActor> GetCombinedGizmoActorClass() const { return CombinedGizmoActorClass; }
+
 protected:
 	UPROPERTY(EditAnywhere, Category="XR Creative")
-	TSubclassOf<AXRCreativeBaseTransformGizmoActor> FullTRSGizmoActorClass;
+	TSubclassOf<AXRCreativeCombinedTransformGizmoActor> CombinedGizmoActorClass;
 
 	UPROPERTY(EditAnywhere, BlueprintReadWrite, Category="XR Creative")
 	TArray<TSubclassOf<AActor>> UnselectableActorClasses;
