XRCreative
- Added Owner Avatar member to XRCreativeGizmo
- Added BP getter on OwnerAvatar member, enabling scripting of behaviours with reference to the XR Avatar
- Added GetGizmoActor function to the tools component

[RN] Added scripting access to XRCreative Custom gizmos

diff --git a/Engine/Plugins/Experimental/XRCreativeFramework/Source/XRCreative/Private/ITF/XRCreativeITFComponent.cpp b/Engine/Plugins/Experimental/XRCreativeFramework/Source/XRCreative/Private/ITF/XRCreativeITFComponent.cpp
--- a/Engine/Plugins/Experimental/XRCreativeFramework/Source/XRCreative/Private/ITF/XRCreativeITFComponent.cpp
+++ b/Engine/Plugins/Experimental/XRCreativeFramework/Source/XRCreative/Private/ITF/XRCreativeITFComponent.cpp
@@ -391,6 +391,10 @@ public:
 			GizmoMesh->Initialize();
 		}
 
+		ITFComponent->GizmoActor = Gizmo;
+
+		Gizmo->OwnerAvatar = Cast<AXRCreativeAvatar>(ITFComponent->GetOwner());
+		
 		return Gizmo;
 	}
 
diff --git a/Engine/Plugins/Experimental/XRCreativeFramework/Source/XRCreative/Public/XRCreativeGizmos.h b/Engine/Plugins/Experimental/XRCreativeFramework/Source/XRCreative/Public/XRCreativeGizmos.h
--- a/Engine/Plugins/Experimental/XRCreativeFramework/Source/XRCreative/Public/XRCreativeGizmos.h
+++ b/Engine/Plugins/Experimental/XRCreativeFramework/Source/XRCreative/Public/XRCreativeGizmos.h
@@ -7,6 +7,8 @@
 #include "XRCreativeGizmos.generated.h"
 
 
+class AXRCreativeAvatar;
+
 namespace UE::GizmoUtil
 {
 	struct FTransformSubGizmoCommonParams;
@@ -74,6 +76,12 @@ public:
 
 	TWeakObjectPtr<UInteractiveGizmoManager> WeakGizmoManager;
 
+	TWeakObjectPtr<AXRCreativeAvatar> OwnerAvatar;
+	
+	/* Get the XRCreative Avatar that spawned this Gizmo */
+	UFUNCTION(BlueprintCallable, Category="XRCreative|Gizmo")
+	AXRCreativeAvatar* GetOwnerAvatar() { return OwnerAvatar.Get(); }
+
 public:
 	// NOTE: These properties alias ones inherited from ACombinedTransformGizmoActor
 	// as a workaround to expose them as EditAnywhere + BlueprintReadWrite without
@@ -165,7 +173,7 @@ public:
 
 	/** XY Plane Scale Component */
 	UPROPERTY(EditAnywhere, BlueprintReadWrite, Category="Gizmo")
-	TObjectPtr<UXRCreativeGizmoMeshComponent> XRPlaneScaleXY;
+	TObjectPtr<UXRCreativeGizmoMeshComponent> XRPlaneScaleXY;	
 };
 
 
diff --git a/Engine/Plugins/Experimental/XRCreativeFramework/Source/XRCreative/Public/XRCreativeITFComponent.h b/Engine/Plugins/Experimental/XRCreativeFramework/Source/XRCreative/Public/XRCreativeITFComponent.h
--- a/Engine/Plugins/Experimental/XRCreativeFramework/Source/XRCreative/Public/XRCreativeITFComponent.h
+++ b/Engine/Plugins/Experimental/XRCreativeFramework/Source/XRCreative/Public/XRCreativeITFComponent.h
@@ -93,7 +93,12 @@ public:
 	UFUNCTION(BlueprintCallable, Category="XR Creative|Tools")
 	void SetCurrentTransformGizmoMode(EToolContextTransformGizmoMode GizmoMode);
 
+	UPROPERTY()
+	TWeakObjectPtr<AXRCreativeCombinedTransformGizmoActor> GizmoActor;
 
+	UFUNCTION(BlueprintCallable, Category="XR Creative|Tools")
+	AXRCreativeCombinedTransformGizmoActor* GetGizmoActor() const { return GizmoActor.Get(); }
+	
 	TSubclassOf<AXRCreativeCombinedTransformGizmoActor> GetCombinedGizmoActorClass() const { return CombinedGizmoActorClass; }
 
 protected:
