XRCreative

- Removed tick functions from XRCreativeGizmo
- Removed ShouldTickIfViewPortsOnly which was causing tick spam in preview editor

#rnx

diff --git a/Engine/Plugins/Experimental/XRCreativeFramework/Source/XRCreative/Private/ITF/XRCreativeGizmos.cpp b/Engine/Plugins/Experimental/XRCreativeFramework/Source/XRCreative/Private/ITF/XRCreativeGizmos.cpp
--- a/Engine/Plugins/Experimental/XRCreativeFramework/Source/XRCreative/Private/ITF/XRCreativeGizmos.cpp
+++ b/Engine/Plugins/Experimental/XRCreativeFramework/Source/XRCreative/Private/ITF/XRCreativeGizmos.cpp
@@ -452,8 +452,6 @@ FTransform UXRCreativeGizmoMeshComponent::CalcViewDependent(
 
 AXRCreativeCombinedTransformGizmoActor::AXRCreativeCombinedTransformGizmoActor()
 {
-	PrimaryActorTick.bCanEverTick = true;
-
 	XRSceneRoot = RootComponent;
 
 	TranslateX = XRTranslateX = CreateOptionalDefaultSubobject<UXRCreativeGizmoMeshComponent>("TranslateX");
@@ -494,12 +492,6 @@ AXRCreativeCombinedTransformGizmoActor::AXRCreativeCombinedTransformGizmoActor()
 }
 
 
-void AXRCreativeCombinedTransformGizmoActor::Tick(float DeltaSeconds)
-{
-	Super::Tick(DeltaSeconds);
-}
-
-
 void AXRCreativeCombinedTransformGizmoActor::SetEnabledElements(ETransformGizmoSubElements EnableElements)
 {
 	const bool bEnableTranslateAxisX = (EnableElements & ETransformGizmoSubElements::TranslateAxisX) != ETransformGizmoSubElements::None;
diff --git a/Engine/Plugins/Experimental/XRCreativeFramework/Source/XRCreative/Public/XRCreativeGizmos.h b/Engine/Plugins/Experimental/XRCreativeFramework/Source/XRCreative/Public/XRCreativeGizmos.h
--- a/Engine/Plugins/Experimental/XRCreativeFramework/Source/XRCreative/Public/XRCreativeGizmos.h
+++ b/Engine/Plugins/Experimental/XRCreativeFramework/Source/XRCreative/Public/XRCreativeGizmos.h
@@ -65,10 +65,7 @@ class AXRCreativeCombinedTransformGizmoActor : public ACombinedTransformGizmoAct
 
 public:
 	AXRCreativeCombinedTransformGizmoActor();
-
-	virtual void Tick(float DeltaSeconds) override;
-	virtual bool ShouldTickIfViewportsOnly() const override { return true; }
-
+	
 	void SetEnabledElements(ETransformGizmoSubElements EnableElements);
 
 	UPROPERTY(BlueprintReadOnly, Category="Gizmo")
