XRCreative

Changed Experimental to Beta plugin designation

[RN] XRCreative Framework is considered Beta

diff --git a/Engine/Plugins/Experimental/XRCreativeFramework/XRCreativeFramework.uplugin b/Engine/Plugins/Experimental/XRCreativeFramework/XRCreativeFramework.uplugin
--- a/Engine/Plugins/Experimental/XRCreativeFramework/XRCreativeFramework.uplugin
+++ b/Engine/Plugins/Experimental/XRCreativeFramework/XRCreativeFramework.uplugin
@@ -11,8 +11,9 @@
 	"MarketplaceURL": "",
 	"SupportURL": "",
 	"CanContainContent": false,
-	"IsBetaVersion": false,
-	"IsExperimentalVersion": true,
+	
+	"IsBetaVersion": true,
+	"IsExperimentalVersion": false,
 	"Installed": false,
 	"EnabledByDefault": false,
 	"SupportedTargetPlatforms": [
