Remove the check for r.PostProcess.PropagateAlpha in the XR Creative ::ValidateSettings().  This is no longer a requirement to starting a session.

diff --git a/Engine/Plugins/Experimental/XRCreativeFramework/Source/XRCreativeEditor/Private/XRCreativeVREditorMode.cpp b/Engine/Plugins/Experimental/XRCreativeFramework/Source/XRCreativeEditor/Private/XRCreativeVREditorMode.cpp
--- a/Engine/Plugins/Experimental/XRCreativeFramework/Source/XRCreativeEditor/Private/XRCreativeVREditorMode.cpp
+++ b/Engine/Plugins/Experimental/XRCreativeFramework/Source/XRCreativeEditor/Private/XRCreativeVREditorMode.cpp
@@ -285,16 +285,6 @@ bool UXRCreativeVREditorMode::ValidateSettings()
 		bSettingsValid = false;
 	}
 
-	IConsoleManager& ConsoleMgr = IConsoleManager::Get();
-	if (IConsoleVariable* PropagateAlpha = ConsoleMgr.FindConsoleVariable(TEXT("r.PostProcessing.PropagateAlpha")))
-	{
-		if (PropagateAlpha->GetBool())
-		{
-			InvalidSettingNotification(LOCTEXT("InvalidCvarPropagateAlpha", "r.PostProcessing.PropagateAlpha must be disabled (and requires an engine restart)"));
-			bSettingsValid = false;
-		}
-	}
-
 	return bSettingsValid;
 }
 
